Integration of AVCodecs into PixelStreaming and support for new pixel formats

diff --git a/Engine/Plugins/Experimental/PixelStreamingPlayer/Source/PixelStreamingPlayer/Private/PixelStreamingPeerComponent.cpp b/Engine/Plugins/Experimental/PixelStreamingPlayer/Source/PixelStreamingPlayer/Private/PixelStreamingPeerComponent.cpp
new file mode 100644
--- /dev/null
+++ b/Engine/Plugins/Experimental/PixelStreamingPlayer/Source/PixelStreamingPlayer/Private/PixelStreamingPeerComponent.cpp
@@ -0,0 +1,226 @@
+// Copyright Epic Games, Inc. All Rights Reserved.
+
+#include "PixelStreamingPeerComponent.h"
+#include "PixelStreamingPlayerPrivate.h"
+#include "PixelStreamingPeerConnection.h"
+#include "FrameBufferRHI.h"
+#include "FrameBufferI420.h"
+#include "ITextureMediaPlayer.h"
+#include "ITextureMediaPlayerModule.h"
+#include "RTCStatsCollector.h"
+#include "Async/Async.h"
+
+UPixelStreamingPeerComponent::UPixelStreamingPeerComponent(const FObjectInitializer& ObjectInitializer)
+	: Super(ObjectInitializer)
+{
+}
+
+void UPixelStreamingPeerComponent::SetConfig(const FPixelStreamingRTCConfigWrapper& Config)
+{
+	PeerConnection = FPixelStreamingPeerConnection::Create(Config.Config);
+	PeerConnection->SetWebRTCStatsCallback(new rtc::RefCountedObject<UE::PixelStreaming::FRTCStatsCollector>(TEXT("Streamer")));
+
+	if (PeerConnection)
+	{
+		if (VideoSinkPlayer != nullptr)
+		{
+			ITextureMediaPlayerModule* TextureModule = FModuleManager::LoadModulePtr<ITextureMediaPlayerModule>("TextureMediaPlayer");
+			if (TextureModule != nullptr)
+			{
+				TextureMediaPlayer = TextureModule->OpenPlayer(VideoSinkPlayer);
+				if (TextureMediaPlayer != nullptr)
+				{
+					PeerConnection->SetVideoSink(this);
+				}
+			}
+		}
+		else
+		{
+			UE_LOG(LogPixelStreamingPlayer, Error, TEXT("MediaPlayer not found."));
+		}
+
+		PeerConnection->OnEmitIceCandidate.AddLambda([this](const webrtc::IceCandidateInterface* Candidate) {
+			FPixelStreamingIceCandidateWrapper CandidateWrapper(*Candidate);
+			AsyncTask(ENamedThreads::GameThread, [this, CandidateWrapper]() {
+				OnIceCandidate.Broadcast(CandidateWrapper);
+			});
+		});
+		PeerConnection->OnIceStateChanged.AddLambda([this](webrtc::PeerConnectionInterface::IceConnectionState NewState) {
+			AsyncTask(ENamedThreads::GameThread, [this, NewState]() {
+				OnIceConnectionChange(NewState);
+			});
+		});
+	}
+	else
+	{
+		UE_LOG(LogPixelStreamingPlayer, Error, TEXT("SetConfig Failed."));
+	}
+}
+
+FPixelStreamingSessionDescriptionWrapper UPixelStreamingPeerComponent::CreateOffer()
+{
+	if (PeerConnection)
+	{
+		FPixelStreamingSessionDescriptionWrapper Wrapper;
+		FEvent* TaskEvent = FPlatformProcess::GetSynchEventFromPool();
+		const auto OnGeneralFailure = [&TaskEvent](const FString& ErrorMsg) {
+			UE_LOG(LogPixelStreamingPlayer, Error, TEXT("CreateOffer Failed: %s"), *ErrorMsg);
+			TaskEvent->Trigger();
+		};
+		AsyncTask(ENamedThreads::AnyNormalThreadNormalTask, [this, &TaskEvent, &Wrapper, &OnGeneralFailure]() {
+			PeerConnection->CreateOffer(
+				FPixelStreamingPeerConnection::EReceiveMediaOption::All,
+				[&TaskEvent, &Wrapper](const webrtc::SessionDescriptionInterface* Sdp) {
+					// success
+					Wrapper.SDP = Sdp;
+					TaskEvent->Trigger();
+				},
+				OnGeneralFailure);
+		});
+		TaskEvent->Wait();
+		FPlatformProcess::ReturnSynchEventToPool(TaskEvent);
+
+		if (Wrapper.SDP == nullptr)
+		{
+			UE_LOG(LogPixelStreamingPlayer, Error, TEXT("CreateOffer Failed: Timeout"));
+		}
+
+		return Wrapper;
+	}
+	else
+	{
+		UE_LOG(LogPixelStreamingPlayer, Error, TEXT("Failed to create offer. Call SetConfig first."));
+		return {};
+	}
+}
+
+FPixelStreamingSessionDescriptionWrapper UPixelStreamingPeerComponent::CreateAnswer(const FString& Sdp)
+{
+	if (PeerConnection)
+	{
+		FPixelStreamingSessionDescriptionWrapper Wrapper;
+		FEvent* TaskEvent = FPlatformProcess::GetSynchEventFromPool();
+		const auto OnGeneralFailure = [&TaskEvent](const FString& ErrorMsg) {
+			UE_LOG(LogPixelStreamingPlayer, Error, TEXT("CreateAnswer Failed: %s"), *ErrorMsg);
+			TaskEvent->Trigger();
+		};
+		AsyncTask(ENamedThreads::AnyNormalThreadNormalTask, [this, &TaskEvent, &Sdp, &Wrapper, &OnGeneralFailure]() {
+			PeerConnection->ReceiveOffer(
+				Sdp,
+				[this, &TaskEvent, &Wrapper, &OnGeneralFailure]() {
+					// success
+					PeerConnection->CreateAnswer(
+						FPixelStreamingPeerConnection::EReceiveMediaOption::All,
+						[&TaskEvent, &Wrapper](const webrtc::SessionDescriptionInterface* Sdp) {
+							// success
+							Wrapper.SDP = Sdp;
+							TaskEvent->Trigger();
+						},
+						OnGeneralFailure);
+				},
+				OnGeneralFailure);
+		});
+		TaskEvent->Wait();
+		FPlatformProcess::ReturnSynchEventToPool(TaskEvent);
+
+		if (Wrapper.SDP == nullptr)
+		{
+			UE_LOG(LogPixelStreamingPlayer, Error, TEXT("CreateAnswer Failed: Timeout"));
+		}
+
+		return Wrapper;
+	}
+	else
+	{
+		UE_LOG(LogPixelStreamingPlayer, Error, TEXT("Failed to create answer. Call SetConfig first."));
+		return {};
+	}
+}
+
+void UPixelStreamingPeerComponent::ReceiveAnswer(const FString& Offer)
+{
+	if (PeerConnection)
+	{
+		FEvent* TaskEvent = FPlatformProcess::GetSynchEventFromPool();
+		const auto OnGeneralFailure = [&TaskEvent](const FString& ErrorMsg) {
+			UE_LOG(LogPixelStreamingPlayer, Error, TEXT("ReceiveAnswer Failed: %s"), *ErrorMsg);
+			TaskEvent->Trigger();
+		};
+		AsyncTask(ENamedThreads::AnyNormalThreadNormalTask, [this, &TaskEvent, &Offer, &OnGeneralFailure]() {
+			PeerConnection->ReceiveAnswer(
+				Offer,
+				[this, &TaskEvent, &OnGeneralFailure]() {
+					// success
+					TaskEvent->Trigger();
+				},
+				OnGeneralFailure);
+		});
+		TaskEvent->Wait();
+		FPlatformProcess::ReturnSynchEventToPool(TaskEvent);
+	}
+	else
+	{
+		UE_LOG(LogPixelStreamingPlayer, Error, TEXT("Failed to receive answer. Call SetConfig first."));
+	}
+}
+
+void UPixelStreamingPeerComponent::ReceiveIceCandidate(const FPixelStreamingIceCandidateWrapper& Candidate)
+{
+	if (PeerConnection)
+	{
+		PeerConnection->AddRemoteIceCandidate(Candidate.Mid, Candidate.MLineIndex, Candidate.SDP);
+	}
+}
+
+void UPixelStreamingPeerComponent::OnIceConnectionChange(webrtc::PeerConnectionInterface::IceConnectionState NewState)
+{
+	if (NewState == webrtc::PeerConnectionInterface::IceConnectionState::kIceConnectionConnected)
+	{
+		OnIceConnection.Broadcast(0);
+	}
+	else if (NewState == webrtc::PeerConnectionInterface::IceConnectionState::kIceConnectionDisconnected)
+	{
+		OnIceDisconnection.Broadcast();
+	}
+}
+
+void UPixelStreamingPeerComponent::OnFrame(const webrtc::VideoFrame& frame)
+{
+	if (frame.video_frame_buffer()->type() == webrtc::VideoFrameBuffer::Type::kNative)
+	{
+		UE::PixelStreaming::FFrameBufferRHI* const FrameBuffer = static_cast<UE::PixelStreaming::FFrameBufferRHI*>(frame.video_frame_buffer().get());
+
+		if (FrameBuffer != nullptr)
+		{
+			TSharedPtr<FVideoResourceRHI, ESPMode::ThreadSafe> VideoResource = static_cast<UE::PixelStreaming::FFrameBufferRHI*>(frame.video_frame_buffer().get())->GetVideoResource();
+
+			if (TextureMediaPlayer.IsValid() && VideoResource.IsValid())
+			{
+				ResourceQueue.Enqueue(VideoResource);
+
+				// Get off WebRTC thread ASAP
+				AsyncTask(ENamedThreads::AnyHiPriThreadHiPriTask, [this, TMPWeak = TextureMediaPlayer.ToWeakPtr()]() {
+					auto PinnedTextureMediaPlayer = TMPWeak.Pin();
+
+					if (PinnedTextureMediaPlayer.IsValid() && !ResourceQueue.IsEmpty())
+					{
+						TSharedPtr<FVideoResource, ESPMode::ThreadSafe> VideoResource;
+
+						while (!ResourceQueue.IsEmpty())
+						{
+							ResourceQueue.Dequeue(VideoResource);
+						}
+
+						auto& Raw = StaticCastSharedPtr<FVideoResourceRHI>(VideoResource)->GetRaw();
+						PinnedTextureMediaPlayer->OnFrame(Raw.Texture, Raw.Fence, Raw.FenceValue);
+					}
+				});
+			}
+		}
+	}
+	if (frame.video_frame_buffer()->type() == webrtc::VideoFrameBuffer::Type::kI420)
+	{
+		// TODO (aidan.possemiers)
+		unimplemented();
+	}
+}
diff --git a/Engine/Plugins/Experimental/PixelStreamingPlayer/Source/PixelStreamingPlayer/Private/PixelStreamingSignallingComponent.cpp b/Engine/Plugins/Experimental/PixelStreamingPlayer/Source/PixelStreamingPlayer/Private/PixelStreamingSignallingComponent.cpp
new file mode 100644
--- /dev/null
+++ b/Engine/Plugins/Experimental/PixelStreamingPlayer/Source/PixelStreamingPlayer/Private/PixelStreamingSignallingComponent.cpp
@@ -0,0 +1,133 @@
+// Copyright Epic Games, Inc. All Rights Reserved.
+
+#include "PixelStreamingSignallingComponent.h"
+#include "PixelStreamingSignallingConnection.h"
+#include "WebSocketsModule.h"
+#include "PixelStreamingPlayerPrivate.h"
+#include "StreamMediaSource.h"
+#include <media/base/media_config.h>
+#include <vector>
+
+UPixelStreamingSignallingComponent::UPixelStreamingSignallingComponent(const FObjectInitializer& ObjectInitializer)
+	: Super(ObjectInitializer)
+{
+	SignallingConnection = MakeUnique<FPixelStreamingSignallingConnection>(MakeShareable(this));
+	// I dont disable keep alive here because it causes a deadlock.
+	// This is because the game thread is waiting to load this component but the timer must be set
+	// on the game thread.
+}
+
+void UPixelStreamingSignallingComponent::Connect(const FString& Url)
+{
+	SignallingConnection->SetKeepAlive(false);
+	SignallingConnection->SetAutoReconnect(true);
+	if (MediaSource == nullptr)
+	{
+		SignallingConnection->TryConnect(Url);
+	}
+	else
+	{
+		SignallingConnection->TryConnect(MediaSource->GetUrl());
+	}
+}
+
+void UPixelStreamingSignallingComponent::Disconnect()
+{
+	SignallingConnection->Disconnect();
+}
+
+void UPixelStreamingSignallingComponent::SendOffer(const FPixelStreamingSessionDescriptionWrapper& Offer)
+{
+	if (Offer.SDP == nullptr)
+	{
+		UE_LOG(LogPixelStreamingPlayer, Error, TEXT("Send Offer failed: Offer was null"));
+	}
+	else
+	{
+		SignallingConnection->SendOffer(*Offer.SDP);
+	}
+}
+
+void UPixelStreamingSignallingComponent::SendAnswer(const FPixelStreamingSessionDescriptionWrapper& Answer)
+{
+	if (Answer.SDP == nullptr)
+	{
+		UE_LOG(LogPixelStreamingPlayer, Error, TEXT("Send Answer failed: Answer was null"));
+	}
+	else
+	{
+		SignallingConnection->SendAnswer(*Answer.SDP);
+	}
+}
+
+void UPixelStreamingSignallingComponent::SendIceCandidate(const FPixelStreamingIceCandidateWrapper& CandidateWrapper)
+{
+	SignallingConnection->SendIceCandidate(*CandidateWrapper.ToWebRTC());
+}
+
+void UPixelStreamingSignallingComponent::OnSignallingConnected()
+{
+	OnConnected.Broadcast();
+}
+
+void UPixelStreamingSignallingComponent::OnSignallingDisconnected(int32 StatusCode, const FString& Reason, bool bWasClean)
+{
+	OnDisconnected.Broadcast(StatusCode, Reason, bWasClean);
+}
+
+void UPixelStreamingSignallingComponent::OnSignallingError(const FString& ErrorMsg)
+{
+	OnConnectionError.Broadcast(ErrorMsg);
+}
+
+void UPixelStreamingSignallingComponent::OnSignallingConfig(const webrtc::PeerConnectionInterface::RTCConfiguration& Config)
+{
+	FPixelStreamingRTCConfigWrapper Wrapper;
+	Wrapper.Config = Config;
+	OnConfig.Broadcast(Wrapper);
+}
+
+void UPixelStreamingSignallingComponent::OnSignallingSessionDescription(webrtc::SdpType Type, const FString& Sdp)
+{
+	if (Type == webrtc::SdpType::kOffer)
+	{
+		OnOffer.Broadcast(Sdp);
+	}
+	else if (Type == webrtc::SdpType::kAnswer)
+	{
+		OnAnswer.Broadcast(Sdp);
+	}
+}
+
+void UPixelStreamingSignallingComponent::OnSignallingRemoteIceCandidate(const FString& SdpMid, int SdpMLineIndex, const FString& Sdp)
+{
+	FPixelStreamingIceCandidateWrapper CandidateWrapper(SdpMid, SdpMLineIndex, Sdp);
+	OnIceCandidate.Broadcast(CandidateWrapper);
+}
+
+void UPixelStreamingSignallingComponent::OnSignallingPeerDataChannels(int32 SendStreamId, int32 RecvStreamId)
+{
+	// TODO (Matthew.Cotton): What to do with data channels?
+}
+
+void UPixelStreamingSignallingComponent::OnSignallingPlayerCount(uint32 Count)
+{
+	// no-op player count is not interesting to us here
+}
+
+void UPixelStreamingSignallingComponent::OnSignallingPlayerConnected(FPixelStreamingPlayerId PlayerId, const FPixelStreamingPlayerConfig& PlayerConfig)
+{
+	// no-op our player receiving other players is possible but not a case we are interested in
+	UE_LOG(LogPixelStreamingPlayer, Warning, TEXT("Warning: Received message about player connected from SS, this very likely means you have misconfigured the Pixel Streaming player connection and entered the streamer port. Please check your connect node and ensure it using the player port, e.g. 80 (or no port)."));
+}
+
+void UPixelStreamingSignallingComponent::OnSignallingPlayerDisconnected(FPixelStreamingPlayerId PlayerId)
+{
+	// no-op our player knowing when other players disconnect is not interesting for now
+	UE_LOG(LogPixelStreamingPlayer, Warning, TEXT("Warning: Received message about player disconnected from SS, this very likely means you have misconfigured the Pixel Streaming player connection and entered the streamer port. Please check your connect node and ensure it using the player port, e.g. 80 (or no port)."));
+}
+
+void UPixelStreamingSignallingComponent::OnSignallingSFUPeerDataChannels(FPixelStreamingPlayerId SFUId, FPixelStreamingPlayerId PlayerId, int32 SendStreamId, int32 RecvStreamId)
+{
+	// TODO (Matthew.Cotton): What to do with data channels?
+}
diff --git a/Engine/Plugins/Experimental/TextureMediaPlayer/Source/TextureMediaPlayer/Private/TextureMediaPlayer.cpp b/Engine/Plugins/Experimental/TextureMediaPlayer/Source/TextureMediaPlayer/Private/TextureMediaPlayer.cpp
new file mode 100644
--- /dev/null
+++ b/Engine/Plugins/Experimental/TextureMediaPlayer/Source/TextureMediaPlayer/Private/TextureMediaPlayer.cpp
@@ -0,0 +1,499 @@
+// Copyright Epic Games, Inc. All Rights Reserved.
+
+#include "TextureMediaPlayer.h"
+
+#include "IMediaEventSink.h"
+#include "IMediaOptions.h"
+#include "ITextureMediaPlayerModule.h"
+#include "MediaSamples.h"
+#include "MediaPlayerOptions.h"
+#include "RenderTargetPool.h"
+#include "TextureMediaPlayerVideoDecoderOutput.h"
+
+DEFINE_LOG_CATEGORY(LogTextureMediaPlayer);
+
+//-----------------------------------------------------------------------------
+/**
+ *	Construction of new player
+ */
+FTextureMediaPlayer::FTextureMediaPlayer(IMediaEventSink& InEventSink)
+	: EventSink(InEventSink)
+	, OptionInterface(nullptr)
+	, State(EMediaState::Closed)
+	, Status(EMediaStatus::None)
+	, bWasClosedOnError(false)
+	, bAllowKillAfterCloseEvent(false)
+	, NumTracksVideo(1)
+	, CurrentTime(FTimespan::Zero())
+	, TimeSinceLastFrameReceived(FTimespan::Zero())
+	, CurrentFrameCount(0)
+	, MostRecentFrameCount(0)
+{
+	MediaSamples.Reset(new FMediaSamples);
+}
+
+//-----------------------------------------------------------------------------
+/**
+ *	Cleanup destructor
+ */
+FTextureMediaPlayer::~FTextureMediaPlayer()
+{
+	CloseInternal(false);
+}
+
+
+// IMediaPlayer interface
+
+FGuid FTextureMediaPlayer::GetPlayerPluginGUID() const
+{
+	static FGuid PlayerPluginGUID(0x16137521, 0x26364e57, 0xa5b66211, 0x821b9819);
+	return PlayerPluginGUID;
+}
+
+FString FTextureMediaPlayer::GetInfo() const
+{
+	return TEXT("No information available");
+}
+
+IMediaSamples& FTextureMediaPlayer::GetSamples()
+{
+	// NOTE: This lock here is not really protecting the media samples in that we have no control what the caller is doing with the return value.
+	//       All this can do is to temporarily guard the thing for clearing its contents.
+	FScopeLock Lock(&MediaSamplesAccessLock);
+	return *MediaSamples;
+}
+
+FString FTextureMediaPlayer::GetStats() const
+{
+	return TEXT("TextureMediaPlayer: GetStats: <empty>?");
+}
+
+IMediaTracks& FTextureMediaPlayer::GetTracks()
+{
+	return *this;
+}
+
+bool FTextureMediaPlayer::Open(const FString& Url, const IMediaOptions* Options)
+{
+	return Open(Url, Options, nullptr);
+}
+
+bool FTextureMediaPlayer::Open(const FString& _InUrl, const IMediaOptions* Options, const FMediaPlayerOptions* InPlayerOptions)
+{
+	FString InUrl(_InUrl);
+	FString Scheme;
+	FString Location;
+
+	if (!InUrl.Split(TEXT("://"), &Scheme, &Location, ESearchCase::CaseSensitive))
+	{
+		UE_LOG(LogTextureMediaPlayer, Error, TEXT("Invalid URL, cannot parse the scheme: %s"), *InUrl);
+		return false;
+	}
+
+	if (Scheme != TEXT("texture"))
+	{
+		UE_LOG(LogTextureMediaPlayer, Error, TEXT("Invalid URL scheme (%s), only 'texture' is supported"), *Scheme);
+		return false;
+	}
+
+	// Clear out our work variables
+	bWasClosedOnError = false;
+
+	// Remember the option interface to poll for changes during playback.
+	OptionInterface = Options;
+
+	// Get a writable copy of the URL so we can override it for debugging.
+	MediaUrl = InUrl;
+
+	Status = EMediaStatus::Connecting;
+	State = EMediaState::Preparing;
+	DeferredEvents.Enqueue(EMediaEvent::MediaConnecting);
+
+	State = EMediaState::Stopped;
+	DeferredEvents.Enqueue(EMediaEvent::MediaOpened);
+	DeferredEvents.Enqueue(EMediaEvent::TracksChanged);
+
+	// Let the module know about our video sink.
+	ITextureMediaPlayerModule* TextureModule = FModuleManager::GetModulePtr<ITextureMediaPlayerModule>("TextureMediaPlayer");
+	if (TextureModule != nullptr)
+	{
+		TextureModule->RegisterVideoSink(AsShared());
+	}
+
+	return true;
+}
+
+//-----------------------------------------------------------------------------
+/**
+ *
+ */
+bool FTextureMediaPlayer::Open(const TSharedRef<FArchive, ESPMode::ThreadSafe>& Archive, const FString& OriginalUrl, const IMediaOptions* /*Options*/)
+{
+	// we support playback only from an external file, not from a "resource" (e.g. a packaged asset)
+	UE_LOG(LogTextureMediaPlayer, Error, TEXT("[%p] TextureMediaPlayer::Open"), this);
+	unimplemented();
+	return false;
+}
+
+//-----------------------------------------------------------------------------
+/**
+ *	Close / Shutdown player
+ */
+void FTextureMediaPlayer::CloseInternal(bool bKillAfterClose)
+{
+// TODO: check that it's not already closed.
+
+	UE_LOG(LogTextureMediaPlayer, Log, TEXT("[%p] TextureMediaPlayer::Close()"), this);
+
+	// For all intents and purposes the player can be considered closed here now already.
+	State = EMediaState::Closed;
+	MediaUrl = FString();
+
+	DeferredEvents.Enqueue(EMediaEvent::TracksChanged);
+	DeferredEvents.Enqueue(EMediaEvent::MediaClosed);
+
+	bAllowKillAfterCloseEvent = bKillAfterClose;
+}
+
+//-----------------------------------------------------------------------------
+/**
+*	Internal Close / Shutdown player
+*/
+void FTextureMediaPlayer::Close()
+{
+	CloseInternal(true);
+	CurrentTime = FTimespan::Zero();
+	CurrentFrameCount = 0;
+	MostRecentFrameCount = 0;
+}
+
+//-----------------------------------------------------------------------------
+/**
+ *	TODO: This would be a better clock. This is the actual time of the audio renderer.
+ */
+void FTextureMediaPlayer::SetLastAudioRenderedSampleTime(FTimespan SampleTime)
+{
+	// No-op.
+}
+
+//-----------------------------------------------------------------------------
+/**
+ *	Tick the player from the game thread
+ */
+void FTextureMediaPlayer::TickInput(FTimespan DeltaTime, FTimespan Timecode)
+{
+	if (State == EMediaState::Error && !bWasClosedOnError)
+	{
+		CloseInternal(true);
+		// Close has set the state to close but we want it to stick at error.
+		State = EMediaState::Error;
+		bWasClosedOnError = true;
+	}
+
+	// Handle state changes as we receive frames.
+	if (State != EMediaState::Error)	
+	{
+		// Has a new frame arrived?
+		if (MostRecentFrameCount != CurrentFrameCount)
+		{
+			TimeSinceLastFrameReceived = FTimespan::Zero();
+			CurrentFrameCount = MostRecentFrameCount;
+			if (State != EMediaState::Playing)
+			{
+				State = EMediaState::Playing;
+				DeferredEvents.Enqueue(EMediaEvent::PlaybackResumed);
+				UE_LOG(LogTextureMediaPlayer, Log, TEXT("[%p] TextureMediaPlayer: received new video frame, switching to playing"), this);
+			}
+		}
+		else
+		{
+			TimeSinceLastFrameReceived += DeltaTime;
+			const int32 kReceiveTimeoutMillis = 10000;
+			int64 timeSinceMillis = TimeSinceLastFrameReceived.GetTotalMilliseconds();
+			if (timeSinceMillis > kReceiveTimeoutMillis)
+			{
+				if (State == EMediaState::Playing)
+				{
+					State = EMediaState::Paused;
+					DeferredEvents.Enqueue(EMediaEvent::PlaybackSuspended);
+					UE_LOG(LogTextureMediaPlayer, Log, TEXT("[%p] TextureMediaPlayer: did not receive new video frame for %d ms, switching to paused"), this, (int32)timeSinceMillis);
+				}
+			}
+		}
+	}
+
+	if (State == EMediaState::Playing)
+	{
+		CurrentTime = Timecode;
+	}
+
+	// Forward enqueued session events. We do this even with no current internal player to ensure all pending events are sent and none are lost.
+	EMediaEvent Event;
+	while (DeferredEvents.Dequeue(Event))
+	{
+		EventSink.ReceiveMediaEvent(Event);
+	}
+}
+
+//-----------------------------------------------------------------------------
+/**
+	Get special feature flags states
+*/
+bool FTextureMediaPlayer::GetPlayerFeatureFlag(EFeatureFlag flag) const
+{
+	switch(flag)
+	{
+		case EFeatureFlag::AllowShutdownOnClose:
+			return bAllowKillAfterCloseEvent;
+		case EFeatureFlag::UsePlaybackTimingV2:
+			return true;
+		case EFeatureFlag::AlwaysPullNewestVideoFrame:
+			return true;
+		default:
+			break;
+	}
+	return IMediaPlayer::GetPlayerFeatureFlag(flag);
+}
+
+//////////////////////////////////////////////////////////////////////////
+// IMediaControl impl
+
+//-----------------------------------------------------------------------------
+/**
+ *	Currently, we cannot control anything. We play a Live video.
+ */
+bool FTextureMediaPlayer::CanControl(EMediaControl Control) const
+{
+	return false;
+}
+
+//-----------------------------------------------------------------------------
+/**
+ *	Rate is only real-time when playing.
+ */
+float FTextureMediaPlayer::GetRate() const
+{
+	return (GetState() == EMediaState::Playing) ? 1.0f : 0.0f;
+}
+
+//-----------------------------------------------------------------------------
+/**
+ *	Expose player state
+ */
+EMediaState FTextureMediaPlayer::GetState() const
+{
+	return State;
+}
+
+//-----------------------------------------------------------------------------
+/**
+ *	Expose player status
+ */
+EMediaStatus FTextureMediaPlayer::GetStatus() const
+{
+	return Status;
+}
+
+//-----------------------------------------------------------------------------
+/**
+ *	Only return real-time playback for the moment..
+ */
+TRangeSet<float> FTextureMediaPlayer::GetSupportedRates(EMediaRateThinning Thinning) const
+{
+	TRangeSet<float> Res;
+	Res.Add(TRange<float>{1.0f}); // only normal (real-time) playback rate
+	Res.Add(TRange<float>{0.0f}); // and pause
+	return Res;
+}
+
+FTimespan FTextureMediaPlayer::GetTime() const
+{
+	return CurrentTime;
+}
+
+FTimespan FTextureMediaPlayer::GetDuration() const
+{
+	// A Live feed has infinite duration.
+	return FTimespan::MaxValue();
+}
+
+bool FTextureMediaPlayer::SetRate(float Rate)
+{
+	// This has no effect at all, but we say it has succeeded anyway.
+	return true;
+}
+
+bool FTextureMediaPlayer::Seek(const FTimespan& Time)
+{
+	return false;
+}
+
+bool FTextureMediaPlayer::GetAudioTrackFormat(int32 TrackIndex, int32 FormatIndex, FMediaAudioTrackFormat& OutFormat) const
+{
+	return false;
+}
+
+int32 FTextureMediaPlayer::GetNumTracks(EMediaTrackType TrackType) const
+{
+	if (TrackType == EMediaTrackType::Video)
+	{
+		return NumTracksVideo;
+	}
+	return 0;
+}
+
+int32 FTextureMediaPlayer::GetNumTrackFormats(EMediaTrackType TrackType, int32 TrackIndex) const
+{
+	return 0;
+}
+
+int32 FTextureMediaPlayer::GetSelectedTrack(EMediaTrackType TrackType) const
+{
+	return TrackType == EMediaTrackType::Video ? 0 : -1;
+}
+
+FText FTextureMediaPlayer::GetTrackDisplayName(EMediaTrackType TrackType, int32 TrackIndex) const
+{
+	return FText();
+}
+
+int32 FTextureMediaPlayer::GetTrackFormat(EMediaTrackType TrackType, int32 TrackIndex) const
+{
+	// Right now we only have a single format per track so we return format index 0 at all times.
+	return 0;
+}
+
+FString FTextureMediaPlayer::GetTrackLanguage(EMediaTrackType TrackType, int32 TrackIndex) const
+{
+	return TEXT("");
+}
+
+FString FTextureMediaPlayer::GetTrackName(EMediaTrackType TrackType, int32 TrackIndex) const
+{
+	return TEXT("");
+}
+
+bool FTextureMediaPlayer::GetVideoTrackFormat(int32 TrackIndex, int32 FormatIndex, FMediaVideoTrackFormat& OutFormat) const
+{
+	return false;
+}
+
+bool FTextureMediaPlayer::SelectTrack(EMediaTrackType TrackType, int32 TrackIndex)
+{
+	return false;
+}
+
+bool FTextureMediaPlayer::SetTrackFormat(EMediaTrackType TrackType, int32 TrackIndex, int32 FormatIndex)
+{
+	return false;
+}
+
+bool FTextureMediaPlayer::SetVideoTrackFrameRate(int32 TrackIndex, int32 FormatIndex, float FrameRate)
+{
+	return false;
+}
+
+void FTextureMediaPlayer::OnFrame(const TArray<uint8>& TextureBuffer, FIntPoint Size)
+{
+	if (TextureBuffer.Num() > 0)
+	{
+		int32 Width = Size.X;
+		int32 SampleHeight = Size.Y;
+		int32 Height = (SampleHeight * 2) / 3;
+
+		FVideoDecoderOutputPtr NewVideoDecoderOutput = MakeShared<TextureMediaPlayerVideoDecoderOutput, ESPMode::ThreadSafe>();
+		TextureMediaPlayerVideoDecoderOutput* VideoDecoderOutput = static_cast<TextureMediaPlayerVideoDecoderOutput*>(NewVideoDecoderOutput.Get());
+
+		// Set up parameters.
+		TUniquePtr<Electra::FParamDict> OutputBufferSampleProperties(new Electra::FParamDict());
+		Electra::FTimeValue TimeStamp;
+		TimeStamp.SetFromHNS((int64)1);	// Duration as HNS; should not be 0 but is unknown. Use a small value instead.
+		OutputBufferSampleProperties->Set("duration", Electra::FVariantValue(TimeStamp));
+		OutputBufferSampleProperties->Set("pixelfmt", Electra::FVariantValue((int64)EPixelFormat::PF_NV12));
+		OutputBufferSampleProperties->Set("pts", Electra::FVariantValue(TimeStamp));
+		OutputBufferSampleProperties->Set("width", Electra::FVariantValue((int64)Width));
+		OutputBufferSampleProperties->Set("height", Electra::FVariantValue((int64)Height));
+
+		VideoDecoderOutput->Initialize(OutputBufferSampleProperties.Release(), TextureBuffer, Size);
+
+		// We are dealing with real time video here, so it is important we show the most recent frame received.
+		// Older frames that may have accumulated in the queue past a certain point are not relevant any more.
+		// If the application does not get ticked for some reason frames might pile up in the queue because they
+		// do not get removed fast enough. Even though the application does only make use of the most recent
+		// frame when ticked the older frames are only removed then and will have consumed lots of memory
+		// in the meantime.
+		// To prevent this from happening we remove old frames from the queue ourselves right here.
+		int32 NumFramesInQueue = MediaSamples->NumVideoSamples();
+		if (NumFramesInQueue > 2)
+		{
+			// This will cause the queue to get flushed, but only on the next game tick.
+			// Anything we enqueue in the meantime will get flushed as well.
+			// NOTE: We cannot use Internal_PurgeVideoSamplesHint here because if the game loop is not
+			//       being ticked the time it uses for comparison will not have advanced and nothing
+			//       will get purged.
+			MediaSamples->PurgeOutdatedVideoSamples(FMediaTimeStamp(FTimespan::MaxValue()), false);
+			EventSink.ReceiveMediaEvent(EMediaEvent::Internal_ResetForDiscontinuity);
+		}
+
+		FElectraTextureSampleRef TextureSample = OutputTexturePool.AcquireShared();
+		// The texture sample needs a POD pointer for some reason but will internally attach to a threadsafe TSharedPtr via AsShared().
+		// Ownership of the decoded frame is thereby transfered to the texture sample.
+		TextureSample->Initialize(VideoDecoderOutput);
+		++MostRecentFrameCount;
+		MediaSamples->AddVideo(TextureSample);		
+	}
+}
+
+#if PLATFORM_WINDOWS
+void FTextureMediaPlayer::OnFrame(FTextureRHIRef TextureRHIRef, TRefCountPtr<ID3D12Fence> D3DFence, uint64 FenceValue)
+{
+	if (TextureRHIRef.IsValid())
+	{
+		FIntVector SizeXYZ = TextureRHIRef->GetSizeXYZ();
+		int32 Width = SizeXYZ.X;
+		int32 SampleHeight = SizeXYZ.Y;
+		int32 Height = (SampleHeight * 2) / 3;
+		
+		FVideoDecoderOutputPtr NewVideoDecoderOutput = MakeShared<TextureMediaPlayerVideoDecoderOutput, ESPMode::ThreadSafe>();
+		TextureMediaPlayerVideoDecoderOutput* VideoDecoderOutput = static_cast<TextureMediaPlayerVideoDecoderOutput*>(NewVideoDecoderOutput.Get());
+
+		// Set up parameters.
+		TUniquePtr<Electra::FParamDict> OutputBufferSampleProperties(new Electra::FParamDict());
+		Electra::FTimeValue TimeStamp;
+		TimeStamp.SetFromHNS((int64)1);	// Duration as HNS; should not be 0 but is unknown. Use a small value instead.
+		OutputBufferSampleProperties->Set("duration", Electra::FVariantValue(TimeStamp));
+		OutputBufferSampleProperties->Set("pixelfmt", Electra::FVariantValue((int64)EPixelFormat::PF_NV12));
+		OutputBufferSampleProperties->Set("pts", Electra::FVariantValue(TimeStamp));
+		OutputBufferSampleProperties->Set("width", Electra::FVariantValue((int64)Width));
+		OutputBufferSampleProperties->Set("height", Electra::FVariantValue((int64)Height));
+
+		VideoDecoderOutput->Initialize(OutputBufferSampleProperties.Release(), TextureRHIRef, FIntPoint(Width, SampleHeight), D3DFence, FenceValue);
+
+		// We are dealing with real time video here, so it is important we show the most recent frame received.
+		// Older frames that may have accumulated in the queue past a certain point are not relevant any more.
+		// If the application does not get ticked for some reason frames might pile up in the queue because they
+		// do not get removed fast enough. Even though the application does only make use of the most recent
+		// frame when ticked the older frames are only removed then and will have consumed lots of memory
+		// in the meantime.
+		// To prevent this from happening we remove old frames from the queue ourselves right here.
+		int32 NumFramesInQueue = MediaSamples->NumVideoSamples();
+		if (NumFramesInQueue > 2)
+		{
+			// This will cause the queue to get flushed, but only on the next game tick.
+			// Anything we enqueue in the meantime will get flushed as well.
+			// NOTE: We cannot use Internal_PurgeVideoSamplesHint here because if the game loop is not
+			//       being ticked the time it uses for comparison will not have advanced and nothing
+			//       will get purged.
+			MediaSamples->PurgeOutdatedVideoSamples(FMediaTimeStamp(FTimespan::MaxValue()), false);
+			EventSink.ReceiveMediaEvent(EMediaEvent::Internal_ResetForDiscontinuity);
+		}
+
+		FElectraTextureSampleRef TextureSample = OutputTexturePool.AcquireShared();
+		// The texture sample needs a POD pointer for some reason but will internally attach to a threadsafe TSharedPtr via AsShared().
+		// Ownership of the decoded frame is thereby transfered to the texture sample.
+		TextureSample->Initialize(VideoDecoderOutput);
+		++MostRecentFrameCount;
+		MediaSamples->AddVideo(TextureSample);
+	}
+}
+#endif
\ No newline at end of file
diff --git a/Engine/Plugins/Experimental/TextureMediaPlayer/Source/TextureMediaPlayer/Private/TextureMediaPlayer.h b/Engine/Plugins/Experimental/TextureMediaPlayer/Source/TextureMediaPlayer/Private/TextureMediaPlayer.h
new file mode 100644
--- /dev/null
+++ b/Engine/Plugins/Experimental/TextureMediaPlayer/Source/TextureMediaPlayer/Private/TextureMediaPlayer.h
@@ -0,0 +1,183 @@
+// Copyright Epic Games, Inc. All Rights Reserved.
+
+#pragma once
+
+#include "Containers/UnrealString.h"
+#include "Containers/Queue.h"
+#include "ElectraTextureSample.h"
+#include "IMediaCache.h"
+#include "IMediaPlayer.h"
+#include "IMediaView.h"
+#include "IMediaControls.h"
+#include "IMediaTracks.h"
+#include "IMediaEventSink.h"
+#include "ITextureMediaPlayer.h"
+#include "MediaSampleQueue.h"
+#include "Misc/Guid.h"
+#include "Templates/SharedPointer.h"
+#include "TextureMediaPlayerPrivate.h"
+
+class IMediaOptions;
+class FMediaSamples;
+class FArchive;
+
+/**
+ * Implements a media player that can handle Texture.
+ */
+class FTextureMediaPlayer
+	: public IMediaPlayer
+	, protected IMediaCache
+	, protected IMediaView
+	, protected IMediaControls
+	, public IMediaTracks
+	, public ITextureMediaPlayer
+{
+public:
+
+	/**
+	 * Create and initialize a new instance.
+	 *
+	 * @param InEventSink The object that receives media events from this player.
+	**/
+	FTextureMediaPlayer(IMediaEventSink& InEventSink);
+
+	virtual ~FTextureMediaPlayer();
+
+public:
+	// IMediaPlayer impl
+
+	void Close() override;
+
+	IMediaCache& GetCache() override
+	{
+		return *this;
+	}
+
+	IMediaControls& GetControls() override
+	{
+		return *this;
+	}
+
+	FString GetUrl() const override
+	{
+		return MediaUrl;
+	}
+
+	IMediaView& GetView() override
+	{
+		return *this;
+	}
+
+	void SetGuid(const FGuid& Guid) override
+	{
+		PlayerGuid = Guid;
+	}
+
+	FString GetInfo() const override;
+	FGuid GetPlayerPluginGUID() const override;
+	IMediaSamples& GetSamples() override;
+	FString GetStats() const override;
+	IMediaTracks& GetTracks() override;
+	bool Open(const FString& Url, const IMediaOptions* Options) override;
+	bool Open(const FString& Url, const IMediaOptions* Options, const FMediaPlayerOptions* PlayerOptions) override;
+	bool Open(const TSharedRef<FArchive, ESPMode::ThreadSafe>& Archive, const FString& OriginalUrl, const IMediaOptions* Options) override;
+	void TickInput(FTimespan DeltaTime, FTimespan Timecode) override;
+	void SetLastAudioRenderedSampleTime(FTimespan SampleTime) override;
+	bool FlushOnSeekStarted() const override
+	{
+		return true;
+	}
+	bool FlushOnSeekCompleted() const override
+	{
+		return false;
+	}
+
+	bool GetPlayerFeatureFlag(EFeatureFlag flag) const override;
+
+	FTimespan GetTime() const override;
+
+	// ITextureMediaPlayer interface
+	virtual void OnFrame(const TArray<uint8>& TextureBuffer, FIntPoint Size) override; // TODO (aidan.possemiers) 
+#if PLATFORM_WINDOWS
+	virtual void OnFrame(FTextureRHIRef TextureRHIRef, TRefCountPtr<ID3D12Fence> D3DFence, uint64 FenceValue) override;
+#else
+	virtual void OnFrame(FTextureRHIRef TextureRHIRef, FGPUFenceRHIRef Fence, uint64 FenceValue) override { unimplemented(); }; // TODO (aidan.possemiers) 
+#endif // PLATFORM_WINDOWS
+
+private:
+
+	void CloseInternal(bool bKillAfterClose);
+
+	// -------------------------------------------------------------------------------------------------------------------------
+
+	// IMediaControls impl
+	bool CanControl(EMediaControl Control) const override;
+
+	FTimespan GetDuration() const override;
+
+	bool IsLooping() const override
+	{
+		return false;
+	}
+
+	bool SetLooping(bool bLooping) override
+	{
+		return bLooping == false;
+	}
+
+	EMediaState GetState() const override;
+	EMediaStatus GetStatus() const override;
+
+	TRangeSet<float> GetSupportedRates(EMediaRateThinning Thinning) const override;
+	float GetRate() const override;
+	bool SetRate(float Rate) override;
+	bool Seek(const FTimespan& Time) override;
+
+	// From IMediaTracks
+	virtual bool GetAudioTrackFormat(int32 TrackIndex, int32 FormatIndex, FMediaAudioTrackFormat& OutFormat) const override;
+	virtual int32 GetNumTracks(EMediaTrackType TrackType) const override;
+	virtual int32 GetNumTrackFormats(EMediaTrackType TrackType, int32 TrackIndex) const override;
+	virtual int32 GetSelectedTrack(EMediaTrackType TrackType) const override;
+	virtual FText GetTrackDisplayName(EMediaTrackType TrackType, int32 TrackIndex) const override;
+	virtual int32 GetTrackFormat(EMediaTrackType TrackType, int32 TrackIndex) const override;
+	virtual FString GetTrackLanguage(EMediaTrackType TrackType, int32 TrackIndex) const override;
+	virtual FString GetTrackName(EMediaTrackType TrackType, int32 TrackIndex) const override;
+	virtual bool GetVideoTrackFormat(int32 TrackIndex, int32 FormatIndex, FMediaVideoTrackFormat& OutFormat) const override;
+	virtual bool SelectTrack(EMediaTrackType TrackType, int32 TrackIndex) override;
+	virtual bool SetTrackFormat(EMediaTrackType TrackType, int32 TrackIndex, int32 FormatIndex) override;
+	virtual bool SetVideoTrackFrameRate(int32 TrackIndex, int32 FormatIndex, float FrameRate);
+
+private:
+	/** The media event handler. */
+	IMediaEventSink&								EventSink;
+	/** Media player Guid */
+	FGuid											PlayerGuid;
+
+	/** Option interface **/
+	const IMediaOptions*							OptionInterface;
+
+	/**  */
+	TAtomic<EMediaState>							State;
+	TAtomic<EMediaStatus>							Status;
+	bool											bWasClosedOnError;
+	bool											bAllowKillAfterCloseEvent;
+
+	/** Queued events */
+	TQueue<EMediaEvent>								DeferredEvents;
+
+	/** The URL of the currently opened media. */
+	FString											MediaUrl;
+
+	int32											NumTracksVideo;
+
+	FCriticalSection												MediaSamplesAccessLock;
+	TUniquePtr<FMediaSamples>										MediaSamples;
+
+	FElectraTextureSamplePool										OutputTexturePool;
+
+	FTimespan														CurrentTime;
+	FTimespan														TimeSinceLastFrameReceived;
+	uint32															CurrentFrameCount;
+	uint32															MostRecentFrameCount;
+};
+
diff --git a/Engine/Plugins/Experimental/TextureMediaPlayer/Source/TextureMediaPlayer/Private/TextureMediaPlayerModule.cpp b/Engine/Plugins/Experimental/TextureMediaPlayer/Source/TextureMediaPlayer/Private/TextureMediaPlayerModule.cpp
new file mode 100644
--- /dev/null
+++ b/Engine/Plugins/Experimental/TextureMediaPlayer/Source/TextureMediaPlayer/Private/TextureMediaPlayerModule.cpp
@@ -0,0 +1,76 @@
+// Copyright Epic Games, Inc. All Rights Reserved.
+
+#include "ITextureMediaPlayerModule.h"
+#include "MediaPlayer.h"
+#include "MediaPlayerOptions.h"
+#include "StreamMediaSource.h"
+#include "UObject/Package.h"
+#include "UObject/UObjectGlobals.h"
+#include "TextureMediaPlayer.h"
+
+class FTextureMediaPlayerModule: public ITextureMediaPlayerModule
+{
+public:
+
+	void StartupModule() override
+	{
+	}
+
+	void ShutdownModule() override
+	{
+		if (MediaSource != nullptr)
+		{
+			MediaSource->RemoveFromRoot();
+			MediaSource = nullptr;
+		}
+	}
+
+	TSharedPtr<IMediaPlayer, ESPMode::ThreadSafe> CreatePlayer(IMediaEventSink& EventSink) override
+	{
+		return MakeShareable(new FTextureMediaPlayer(EventSink));
+	}
+
+	virtual TSharedPtr<ITextureMediaPlayer, ESPMode::ThreadSafe> OpenPlayer(UMediaPlayer* MediaPlayer) override;
+	virtual void RegisterVideoSink(TSharedPtr<ITextureMediaPlayer, ESPMode::ThreadSafe> InVideoSink) override;
+
+private:
+	/** We use this media source for all the players. */
+	TWeakObjectPtr<UStreamMediaSource> MediaSource;
+	/** Used only to pass the video sink in OpenPlayer. */
+	TSharedPtr<ITextureMediaPlayer, ESPMode::ThreadSafe> VideoSink = nullptr;
+};
+
+TSharedPtr<ITextureMediaPlayer, ESPMode::ThreadSafe> FTextureMediaPlayerModule::OpenPlayer(UMediaPlayer* MediaPlayer)
+{
+	TSharedPtr<ITextureMediaPlayer, ESPMode::ThreadSafe> NewVideoSink = nullptr;
+
+	if (MediaPlayer != nullptr)
+	{
+		// Create the media source if we don't have it yet.
+		if (MediaSource == nullptr)
+		{
+			MediaSource = NewObject<UStreamMediaSource>(GetTransientPackage());
+			MediaSource->AddToRoot();
+			MediaSource->StreamUrl = FString(TEXT("Texture://"));
+		}
+
+		// Open the player.
+		FMediaPlayerOptions Options;
+		Options.PlayOnOpen = EMediaPlayerOptionBooleanOverride::Enabled;
+		MediaPlayer->OpenSourceWithOptions(MediaSource.Get(), Options);
+
+		// Retrive the video sink if it was set.
+		NewVideoSink = VideoSink;
+		VideoSink = nullptr;
+	}
+
+	return NewVideoSink;
+}
+
+void FTextureMediaPlayerModule::RegisterVideoSink(TSharedPtr<ITextureMediaPlayer, ESPMode::ThreadSafe> InVideoSink)
+{
+	VideoSink = InVideoSink;
+}
+
+
+IMPLEMENT_MODULE(FTextureMediaPlayerModule, TextureMediaPlayer)
diff --git a/Engine/Plugins/Experimental/TextureMediaPlayer/Source/TextureMediaPlayer/Private/TextureMediaPlayerPrivate.h b/Engine/Plugins/Experimental/TextureMediaPlayer/Source/TextureMediaPlayer/Private/TextureMediaPlayerPrivate.h
new file mode 100644
--- /dev/null
+++ b/Engine/Plugins/Experimental/TextureMediaPlayer/Source/TextureMediaPlayer/Private/TextureMediaPlayerPrivate.h
@@ -0,0 +1,7 @@
+// Copyright Epic Games, Inc. All Rights Reserved.
+
+#pragma once
+
+#include "Logging/LogMacros.h"
+
+DECLARE_LOG_CATEGORY_EXTERN(LogTextureMediaPlayer, Log, All);
diff --git a/Engine/Plugins/Experimental/TextureMediaPlayer/Source/TextureMediaPlayer/Private/TextureMediaPlayerVideoDecoderOutput.cpp b/Engine/Plugins/Experimental/TextureMediaPlayer/Source/TextureMediaPlayer/Private/TextureMediaPlayerVideoDecoderOutput.cpp
new file mode 100644
--- /dev/null
+++ b/Engine/Plugins/Experimental/TextureMediaPlayer/Source/TextureMediaPlayer/Private/TextureMediaPlayerVideoDecoderOutput.cpp
@@ -0,0 +1,94 @@
+// Copyright Epic Games, Inc. All Rights Reserved.
+
+#include "TextureMediaPlayerVideoDecoderOutput.h"
+
+void TextureMediaPlayerVideoDecoderOutput::Initialize(Electra::FParamDict* InParamDict, const TArray<uint8>& InBuffer, const FIntPoint& InSampleDim)
+{
+	FVideoDecoderOutputPC::Initialize(InParamDict);
+	ByteBuffer = MakeShared<TArray<uint8>, ESPMode::ThreadSafe>(InBuffer);
+	SampleDim = InSampleDim;
+}
+
+#if PLATFORM_WINDOWS
+void TextureMediaPlayerVideoDecoderOutput::Initialize(Electra::FParamDict* InParamDict, FTexture2DRHIRef InTexture, const FIntPoint& InSampleDim, TRefCountPtr<ID3D12Fence> InFence, uint64 InFenceValue)
+#else
+void TextureMediaPlayerVideoDecoderOutput::Initialize(Electra::FParamDict* InParamDict, FTexture2DRHIRef InTexture, const FIntPoint& InSampleDim, FGPUFenceRHIRef InFence, uint64 InFenceValue)
+#endif
+{
+	FVideoDecoderOutputPC::Initialize(InParamDict);
+	Texture = InTexture;
+	SampleDim = InSampleDim;
+	Fence = InFence;
+	FenceValue = InFenceValue;
+}
+
+void TextureMediaPlayerVideoDecoderOutput::SetOwner(const TSharedPtr<IDecoderOutputOwner, ESPMode::ThreadSafe>& Renderer)
+{
+	check(!"Should not be called!");
+}
+
+TextureMediaPlayerVideoDecoderOutput::EOutputType TextureMediaPlayerVideoDecoderOutput::GetOutputType() const
+{
+	if(Texture.IsValid())
+	{
+		return TextureMediaPlayerVideoDecoderOutput::EOutputType::Hardware_DX;
+	}
+
+	return TextureMediaPlayerVideoDecoderOutput::EOutputType::HardwareDX9_DX12;
+}
+
+const TArray<uint8>& TextureMediaPlayerVideoDecoderOutput::GetBuffer() const
+{
+	return *ByteBuffer;
+}
+
+uint32 TextureMediaPlayerVideoDecoderOutput::GetStride() const
+{
+	return 0;
+}
+
+FIntPoint TextureMediaPlayerVideoDecoderOutput::GetDim() const
+{
+	return SampleDim;
+}
+
+#if PLATFORM_WINDOWS
+
+TRefCountPtr<IMFSample> TextureMediaPlayerVideoDecoderOutput::GetMFSample() const
+{
+	check(!"Should not be called!");
+	return nullptr;
+}
+
+// TODO: Use commented lines to replace following lines once Electra changes are integrated
+//TRefCountPtr<IUnknown> TextureMediaPlayerVideoDecoderOutput::GetTexture() const
+TRefCountPtr<ID3D11Texture2D> TextureMediaPlayerVideoDecoderOutput::GetTexture() const
+{
+	//IUnknown* Result = nullptr;
+	ID3D11Texture2D* Result = nullptr;
+	if (Texture.IsValid())
+	{
+		//Result = static_cast<IUnknown*>(Texture->GetNativeResource());
+		Result = static_cast<ID3D11Texture2D*>(Texture->GetNativeResource());
+	}
+	return Result;
+}
+
+TRefCountPtr<IUnknown> TextureMediaPlayerVideoDecoderOutput::GetSync(uint64& SyncValue) const
+{
+	if (Fence.IsValid())
+	{
+		SyncValue = FenceValue;
+		return static_cast<IUnknown*>(Fence.GetReference());
+	}
+	SyncValue = 0;
+	return nullptr;
+}
+
+TRefCountPtr<ID3D11Device> TextureMediaPlayerVideoDecoderOutput::GetDevice() const
+{
+	check(!"Should not be called!");
+	return nullptr;
+}
+
+#endif // PLATFORM_WINDOWS
diff --git a/Engine/Plugins/Experimental/TextureMediaPlayer/Source/TextureMediaPlayer/Private/TextureMediaPlayerVideoDecoderOutput.h b/Engine/Plugins/Experimental/TextureMediaPlayer/Source/TextureMediaPlayer/Private/TextureMediaPlayerVideoDecoderOutput.h
new file mode 100644
--- /dev/null
+++ b/Engine/Plugins/Experimental/TextureMediaPlayer/Source/TextureMediaPlayer/Private/TextureMediaPlayerVideoDecoderOutput.h
@@ -0,0 +1,64 @@
+// Copyright Epic Games, Inc. All Rights Reserved.
+
+#pragma once
+
+#include "MediaVideoDecoderOutputPC.h"
+#include "RHI.h"
+
+#if PLATFORM_WINDOWS
+
+THIRD_PARTY_INCLUDES_START
+#include "d3d12.h"
+THIRD_PARTY_INCLUDES_END
+
+#endif // PLATFORM_WINDOWS
+
+class TextureMediaPlayerVideoDecoderOutput : public FVideoDecoderOutputPC
+{
+public:
+	void Initialize(Electra::FParamDict* InParamDict, const TArray<uint8>& InBuffer, const FIntPoint& InSampleDim);
+
+#if PLATFORM_WINDOWS
+	void Initialize(Electra::FParamDict* InParamDict, FTexture2DRHIRef InTexture, const FIntPoint& InSampleDim, TRefCountPtr<ID3D12Fence> InFence, uint64 InFenceValue);
+#else
+	void Initialize(Electra::FParamDict* InParamDict, FTexture2DRHIRef InTexture, const FIntPoint& InSampleDim, FGPUFenceRHIRef InFence, uint64 InFenceValue);
+#endif // PLATFORM_WINDOWS
+
+	// FVideoDecoderOutputPC interface.
+	virtual void SetOwner(const TSharedPtr<IDecoderOutputOwner, ESPMode::ThreadSafe>& Renderer) override;
+
+	virtual EOutputType GetOutputType() const override;
+
+	virtual const TArray<uint8>& GetBuffer() const override;
+	virtual uint32 GetStride() const override;
+	virtual FIntPoint GetDim() const;
+
+#if PLATFORM_WINDOWS
+	virtual TRefCountPtr<IMFSample> GetMFSample() const override;
+
+	// TODO: Use commented line to replace following line once Electra changes are integrated
+	//virtual TRefCountPtr<IUnknown> GetTexture() const override;
+	virtual TRefCountPtr<ID3D11Texture2D> GetTexture() const override;
+	// TODO: Re-add override specifier once Electra changes are integrated
+	virtual TRefCountPtr<IUnknown> GetSync(uint64& SyncValue) const;
+	virtual TRefCountPtr<ID3D11Device> GetDevice() const override;
+#endif // PLATFORM_WINDOWS
+
+private:
+	/** Texture from WebRTC with accompaning sync data. */
+	FTexture2DRHIRef Texture;
+	
+#if PLATFORM_WINDOWS
+	TRefCountPtr<ID3D12Fence> Fence;
+#else
+	FGPUFenceRHIRef Fence;
+#endif // PLATFORM_WINDOWS
+
+	uint64 FenceValue;
+
+	/** Dimension of any internally allocated buffer - stored explicitly to cover various special cases for DX. */
+	FIntPoint SampleDim;
+
+	/** Not used but I need to return a buffer in GetBuffer. */
+	TSharedPtr<TArray<uint8>, ESPMode::ThreadSafe> ByteBuffer;
+};
diff --git a/Engine/Plugins/Experimental/TextureMediaPlayer/Source/TextureMediaPlayer/Public/ITextureMediaPlayer.h b/Engine/Plugins/Experimental/TextureMediaPlayer/Source/TextureMediaPlayer/Public/ITextureMediaPlayer.h
new file mode 100644
--- /dev/null
+++ b/Engine/Plugins/Experimental/TextureMediaPlayer/Source/TextureMediaPlayer/Public/ITextureMediaPlayer.h
@@ -0,0 +1,28 @@
+// Copyright Epic Games, Inc. All Rights Reserved.
+
+#pragma once
+
+#include "RHI.h"
+
+#if PLATFORM_WINDOWS
+
+THIRD_PARTY_INCLUDES_START
+#include "d3d12.h"
+THIRD_PARTY_INCLUDES_END
+
+#endif // PLATFORM_WINDOWS
+
+class ITextureMediaPlayer : public TSharedFromThis<ITextureMediaPlayer, ESPMode::ThreadSafe>
+{
+public:
+	virtual ~ITextureMediaPlayer() {}
+
+	virtual void OnFrame(const TArray<uint8>& TextureBuffer, FIntPoint Size) = 0;
+
+#if PLATFORM_WINDOWS
+	virtual void OnFrame(FTextureRHIRef TextureRHIRef, TRefCountPtr<ID3D12Fence> D3DFence, uint64 FenceValue) = 0;
+#else
+	virtual void OnFrame(FTextureRHIRef TextureRHIRef, FGPUFenceRHIRef Fence, uint64 FenceValue) = 0;
+#endif // PLATFORM_WINDOWS
+
+};
diff --git a/Engine/Plugins/Experimental/TextureMediaPlayer/Source/TextureMediaPlayer/Public/ITextureMediaPlayerModule.h b/Engine/Plugins/Experimental/TextureMediaPlayer/Source/TextureMediaPlayer/Public/ITextureMediaPlayerModule.h
new file mode 100644
--- /dev/null
+++ b/Engine/Plugins/Experimental/TextureMediaPlayer/Source/TextureMediaPlayer/Public/ITextureMediaPlayerModule.h
@@ -0,0 +1,29 @@
+// Copyright Epic Games, Inc. All Rights Reserved.
+
+#pragma once
+
+#include "Modules/ModuleManager.h"
+#include "RHI.h"
+
+class ITextureMediaPlayer;
+class IMediaEventSink;
+class IMediaPlayer;
+class UMediaPlayer;
+
+class ITextureMediaPlayerModule: public IModuleInterface
+{
+public:
+	virtual TSharedPtr<IMediaPlayer, ESPMode::ThreadSafe> CreatePlayer(IMediaEventSink& EventSink) = 0;
+
+	/**
+	 * Opens the player and returns the interface which you can then pass to Texture
+	 * so the player can receive frames from Texture.
+	 */
+	virtual TSharedPtr<ITextureMediaPlayer, ESPMode::ThreadSafe> OpenPlayer(UMediaPlayer* MediaPlayer) = 0;
+
+	/**
+	 * Players will call this so that this module can return this video sink in OpenPlayer.
+	 */
+	virtual void RegisterVideoSink(TSharedPtr<ITextureMediaPlayer, ESPMode::ThreadSafe> InVideoSink) = 0;
+
+};
diff --git a/Engine/Plugins/Experimental/TextureMediaPlayer/Source/TextureMediaPlayer/TextureMediaPlayer.Build.cs b/Engine/Plugins/Experimental/TextureMediaPlayer/Source/TextureMediaPlayer/TextureMediaPlayer.Build.cs
new file mode 100644
--- /dev/null
+++ b/Engine/Plugins/Experimental/TextureMediaPlayer/Source/TextureMediaPlayer/TextureMediaPlayer.Build.cs
@@ -0,0 +1,44 @@
+// Copyright Epic Games, Inc. All Rights Reserved.
+
+using System;
+using System.IO;
+using System.Collections.Generic;
+using EpicGames.Core;
+
+namespace UnrealBuildTool.Rules
+{
+	public class TextureMediaPlayer : ModuleRules
+	{
+		public TextureMediaPlayer(ReadOnlyTargetRules Target) : base(Target)
+		{
+			PrivateDependencyModuleNames.AddRange(new string[]
+			{
+				"Core",
+				"CoreUObject",
+				"ElectraBase",
+				"ElectraSamples",
+				"Engine",
+				"HTTP",
+				"MediaAssets",
+				"MediaUtils",
+				"RenderCore",
+				"RHI",
+			});
+
+			DynamicallyLoadedModuleNames.AddRange(new string[]
+			{
+				"Media",
+			});
+
+			PrivateIncludePathModuleNames.AddRange(new string[]
+			{
+				"Media",
+			});
+
+			if (Target.bCompileAgainstEngine)
+			{
+				PrivateDependencyModuleNames.Add("Engine");
+			}
+		}
+	}
+}
diff --git a/Engine/Plugins/Experimental/TextureMediaPlayer/Source/TextureMediaPlayerFactory/Private/TextureMediaPlayerFactoryModule.cpp b/Engine/Plugins/Experimental/TextureMediaPlayer/Source/TextureMediaPlayerFactory/Private/TextureMediaPlayerFactoryModule.cpp
new file mode 100644
--- /dev/null
+++ b/Engine/Plugins/Experimental/TextureMediaPlayer/Source/TextureMediaPlayerFactory/Private/TextureMediaPlayerFactoryModule.cpp
@@ -0,0 +1,141 @@
+// Copyright Epic Games, Inc. All Rights Reserved.
+
+#include "Containers/Array.h"
+#include "Containers/UnrealString.h"
+#include "IMediaOptions.h"
+#include "IMediaPlayerFactory.h"
+#include "IMediaModule.h"
+#include "Internationalization/Internationalization.h"
+#include "ITextureMediaPlayerModule.h"
+#include "Logging/LogMacros.h"
+#include "Misc/Guid.h"
+#include "Misc/Paths.h"
+#include "Modules/ModuleManager.h"
+#include "UObject/NameTypes.h"
+
+#define LOCTEXT_NAMESPACE "FTextureMediaPlayerFactoryModule"
+
+class FFactoryModule
+	: public IMediaPlayerFactory
+	, public IModuleInterface
+{
+public:
+	FFactoryModule() { }
+
+public:
+	// IMediaPlayerFactory interface
+
+	bool CanPlayUrl(const FString& Url, const IMediaOptions* Options, TArray<FText>* OutWarnings, TArray<FText>* OutErrors) const override
+	{
+		FString Scheme;
+		FString Location;
+
+		// check scheme
+		if (!Url.Split(TEXT("://"), &Scheme, &Location, ESearchCase::CaseSensitive))
+		{
+			if (OutErrors != nullptr)
+			{
+				OutErrors->Add(LOCTEXT("NoSchemeFound", "No URI scheme found"));
+			}
+
+			return false;
+		}
+
+		if (!SupportedUriSchemes.Contains(Scheme))
+		{
+			if (OutErrors != nullptr)
+			{
+				OutErrors->Add(FText::Format(LOCTEXT("SchemeNotSupported", "The URI scheme '{0}' is not supported"), FText::FromString(Scheme)));
+			}
+
+			return false;
+		}
+
+		return true;
+	}
+
+	TSharedPtr<IMediaPlayer, ESPMode::ThreadSafe> CreatePlayer(IMediaEventSink& EventSink) override
+	{
+		auto TextureMediaPlayerModule = FModuleManager::LoadModulePtr<ITextureMediaPlayerModule>("TextureMediaPlayer");
+		return (TextureMediaPlayerModule != nullptr) ? TextureMediaPlayerModule->CreatePlayer(EventSink) : nullptr;
+	}
+
+	FText GetDisplayName() const override
+	{
+		return LOCTEXT("MediaPlayerDisplayName", "Texture Media Player");
+	}
+
+	virtual FName GetPlayerName() const override
+	{
+		static FName PlayerName(TEXT("TextureMediaPlayer"));
+		return PlayerName;
+	}
+
+	virtual FGuid GetPlayerPluginGUID() const override
+	{
+		static FGuid PlayerPluginGUID(0x16137521, 0x26364e57, 0xa5b66211, 0x821b9819);
+		return PlayerPluginGUID;
+	}
+
+	const TArray<FString>& GetSupportedPlatforms() const override
+	{
+		return SupportedPlatforms;
+	}
+
+	bool SupportsFeature(EMediaFeature Feature) const override
+	{
+		return (Feature == EMediaFeature::VideoSamples);
+	}
+
+public:
+
+	// IModuleInterface interface
+
+	void StartupModule() override
+	{
+		// We make sure that the correct modules are loaded.
+		FModuleManager::Get().LoadModule(TEXT("TextureMediaPlayer"));
+
+		// supported platforms
+		SupportedPlatforms.Add(TEXT("Windows"));
+		SupportedPlatforms.Add(TEXT("PS4"));
+		SupportedPlatforms.Add(TEXT("Switch"));
+		SupportedPlatforms.Add(TEXT("XboxOne"));
+		SupportedPlatforms.Add(TEXT("XboxOneGDK"));
+		SupportedPlatforms.Add(TEXT("PS5"));
+		SupportedPlatforms.Add(TEXT("XSX"));
+
+		// supported schemes
+		SupportedUriSchemes.Add(TEXT("texture"));
+
+		// register player factory
+		auto MediaModule = FModuleManager::LoadModulePtr<IMediaModule>("Media");
+
+		if (MediaModule != nullptr)
+		{
+			MediaModule->RegisterPlayerFactory(*this);
+		}
+	}
+
+	void ShutdownModule() override
+	{
+		// unregister player factory
+		auto MediaModule = FModuleManager::GetModulePtr<IMediaModule>("Media");
+
+		if (MediaModule != nullptr)
+		{
+			MediaModule->UnregisterPlayerFactory(*this);
+		}
+	}
+
+private:
+	/** List of platforms that the media player support. */
+	TArray<FString> SupportedPlatforms;
+
+	/** List of supported URI schemes. */
+	TArray<FString> SupportedUriSchemes;
+};
+
+#undef LOCTEXT_NAMESPACE
+
+IMPLEMENT_MODULE(FFactoryModule, TextureMediaPlayerFactory);
diff --git a/Engine/Plugins/Experimental/TextureMediaPlayer/Source/TextureMediaPlayerFactory/TextureMediaPlayerFactory.Build.cs b/Engine/Plugins/Experimental/TextureMediaPlayer/Source/TextureMediaPlayerFactory/TextureMediaPlayerFactory.Build.cs
new file mode 100644
--- /dev/null
+++ b/Engine/Plugins/Experimental/TextureMediaPlayer/Source/TextureMediaPlayerFactory/TextureMediaPlayerFactory.Build.cs
@@ -0,0 +1,29 @@
+// Copyright Epic Games, Inc. All Rights Reserved.
+
+namespace UnrealBuildTool.Rules
+{
+	public class TextureMediaPlayerFactory : ModuleRules
+	{
+		public TextureMediaPlayerFactory(ReadOnlyTargetRules Target) : base(Target)
+		{
+			DynamicallyLoadedModuleNames.AddRange(
+				new string[] {
+					"Media",
+				});
+
+			PrivateDependencyModuleNames.AddRange(
+				new string[] {
+					"Core",
+					"CoreUObject",
+					"MediaAssets",
+                });
+
+			PrivateIncludePathModuleNames.AddRange(
+				new string[] {
+					"Media",
+				});
+
+			PublicIncludePathModuleNames.Add("TextureMediaPlayer");
+		}
+	}
+}
diff --git a/Engine/Plugins/Experimental/TextureMediaPlayer/TextureMediaPlayer.uplugin b/Engine/Plugins/Experimental/TextureMediaPlayer/TextureMediaPlayer.uplugin
new file mode 100644
--- /dev/null
+++ b/Engine/Plugins/Experimental/TextureMediaPlayer/TextureMediaPlayer.uplugin
@@ -0,0 +1,46 @@
+{
+	"FileVersion": 3,
+	"Version": 0,
+	"VersionName": "0.1",
+	"FriendlyName": "Texture Media Player",
+	"Description": "Implements a media player for textures.",
+	"Category": "Media Players",
+	"CreatedBy": "Epic Games, Inc.",
+	"CreatedByURL": "http://epicgames.com",
+	"DocsURL": "",
+	"MarketplaceURL": "",
+	"SupportURL": "",
+	"EnabledByDefault": false,
+	"CanContainContent": false,
+	"IsBetaVersion": true,
+	"Installed": false,
+	"Modules": [
+		{
+			"Name": "TextureMediaPlayer",
+			"Type": "ClientOnlyNoCommandlet",
+			"LoadingPhase": "PostEngineInit",
+			"TargetDenyList": [ "Server" ],
+			"PlatformAllowList": [ "Win64" ]
+		},
+		{
+			"Name": "TextureMediaPlayerFactory",
+			"Type": "ClientOnlyNoCommandlet",
+			"LoadingPhase": "PostEngineInit",
+			"TargetDenyList": [ "Server" ],
+			"PlatformAllowList": [ "Win64" ]
+		},
+		{
+			"Name": "TextureMediaPlayerFactory",
+			"Type": "Editor",
+			"LoadingPhase": "PostEngineInit",
+			"TargetDenyList": [ "Server" ],
+			"PlatformAllowList": [ "Win64" ]
+		}
+	],
+	"Plugins": [
+		{
+			"Name": "ElectraUtil",
+			"Enabled": true
+		},
+	]
+}
diff --git a/Engine/Plugins/Experimental/VirtualCamera/Source/PixelStreamingVCam/PixelStreamingVCam.Build.cs b/Engine/Plugins/Experimental/VirtualCamera/Source/PixelStreamingVCam/PixelStreamingVCam.Build.cs
--- a/Engine/Plugins/Experimental/VirtualCamera/Source/PixelStreamingVCam/PixelStreamingVCam.Build.cs
+++ b/Engine/Plugins/Experimental/VirtualCamera/Source/PixelStreamingVCam/PixelStreamingVCam.Build.cs
@@ -33,6 +33,7 @@ namespace UnrealBuildTool.Rules
 				"LiveLinkInterface",
 				"PixelCapture",
 				"PixelStreamingEditor",
+				"PixelStreamingInput",
 				"SlateCore",
 				"Slate",
 				"VPUtilities",
diff --git a/Engine/Plugins/Experimental/VirtualCamera/Source/PixelStreamingVCam/Private/PixelStreamingMediaOutput.cpp b/Engine/Plugins/Experimental/VirtualCamera/Source/PixelStreamingVCam/Private/PixelStreamingMediaOutput.cpp
--- a/Engine/Plugins/Experimental/VirtualCamera/Source/PixelStreamingVCam/Private/PixelStreamingMediaOutput.cpp
+++ b/Engine/Plugins/Experimental/VirtualCamera/Source/PixelStreamingVCam/Private/PixelStreamingMediaOutput.cpp
@@ -46,27 +46,28 @@ UMediaCapture* UPixelStreamingMediaOutput::CreateMediaCaptureImpl()
 
 void UPixelStreamingMediaOutput::RegisterRemoteResolutionCommandHandler()
 {
+	// TODO: Complete integration work from dev stream
 	// Override resolution command as we this to set the output provider override resolution
-	TSharedPtr<IPixelStreamingInputHandler> InputHandler = Streamer->GetInputHandler().Pin();
-	if(InputHandler)
-	{
-		InputHandler->SetCommandHandler(TEXT("Resolution.Width"), [this](FString Descriptor, FString WidthString){
-			bool bSuccess = false;
-			FString HeightString;
-			UE::PixelStreaming::ExtractJsonFromDescriptor(Descriptor, TEXT("Resolution.Height"), HeightString, bSuccess);
-			if (bSuccess)
-			{
-				int Width = FCString::Atoi(*WidthString);
-				int Height = FCString::Atoi(*HeightString);
-				if (Width < 1 || Height < 1)
-				{
-					return;
-				}
+	// TSharedPtr<IPixelStreamingInputHandler> InputHandler = Streamer->GetInputHandler().Pin();
+	// if(InputHandler)
+	// {
+	// 	InputHandler->SetCommandHandler(TEXT("Resolution.Width"), [this](FString Descriptor, FString WidthString){
+	// 		bool bSuccess = false;
+	// 		FString HeightString;
+	// 		UE::PixelStreaming::ExtractJsonFromDescriptor(Descriptor, TEXT("Resolution.Height"), HeightString, bSuccess);
+	// 		if (bSuccess)
+	// 		{
+	// 			int Width = FCString::Atoi(*WidthString);
+	// 			int Height = FCString::Atoi(*HeightString);
+	// 			if (Width < 1 || Height < 1)
+	// 			{
+	// 				return;
+	// 			}
 
-				RemoteResolutionChangedEvent.Broadcast(FIntPoint(Width, Height));
-			}
-		});
-	}
+	// 			RemoteResolutionChangedEvent.Broadcast(FIntPoint(Width, Height));
+	// 		}
+	// 	});
+	// }
 }
 
 void UPixelStreamingMediaOutput::StartStreaming()
diff --git a/Engine/Plugins/Experimental/VirtualCamera/Source/PixelStreamingVCam/Private/VCamPixelStreamingSession.cpp b/Engine/Plugins/Experimental/VirtualCamera/Source/PixelStreamingVCam/Private/VCamPixelStreamingSession.cpp
--- a/Engine/Plugins/Experimental/VirtualCamera/Source/PixelStreamingVCam/Private/VCamPixelStreamingSession.cpp
+++ b/Engine/Plugins/Experimental/VirtualCamera/Source/PixelStreamingVCam/Private/VCamPixelStreamingSession.cpp
@@ -9,11 +9,13 @@
 #include "Editor/EditorPerformanceSettings.h"
 #include "IPixelStreamingInputHandler.h"
 #include "IPixelStreamingModule.h"
+#include "IPixelStreamingInputModule.h"
 #include "Modules/ModuleManager.h"
 #include "Math/Matrix.h"
 #include "PixelStreamingEditorModule.h"
 #include "PixelStreamingVCamLog.h"
-#include "PixelStreamingProtocol.h"
+#include "PixelStreamingInputProtocol.h"
+#include "PixelStreamingInputMessage.h"
 #include "PixelStreamingServers.h"
 #include "Serialization/MemoryReader.h"
 #include "VCamComponent.h"
@@ -21,11 +23,12 @@
 #include "VCamPixelStreamingLiveLink.h"
 #include "VPFullScreenUserWidget.h"
 #include "Widgets/SVirtualWindow.h"
+#include "PixelStreamingInputEnums.h"
 
 namespace UE::VCamPixelStreamingSession::Private
 {
 	static const FSoftClassPath EmptyUMGSoftClassPath(TEXT("/VCamCore/Assets/VCam_EmptyVisibleUMG.VCam_EmptyVisibleUMG_C"));
-} // namespace VCamPixelStreamingSession
+} // namespace UE::VCamPixelStreamingSession::Private
 
 void UVCamPixelStreamingSession::Initialize()
 {
@@ -103,14 +106,13 @@ void UVCamPixelStreamingSession::Activate()
 	int32 StreamerPort = FPixelStreamingEditorModule::GetModule()->GetStreamerPort();
 	MediaOutput->SetSignallingServerURL(FString::Printf(TEXT("%s:%s"), *SignallingDomain, *FString::FromInt(StreamerPort)));
 	UE_LOG(LogPixelStreamingVCam, Log, TEXT("Activating PixelStreaming VCam Session. Endpoint: %s:%s"), *SignallingDomain, *FString::FromInt(StreamerPort));
-	
 }
 
 void UVCamPixelStreamingSession::SetupCapture()
 {
 	UE_LOG(LogPixelStreamingVCam, Log, TEXT("Create new media capture for Pixel Streaming VCam."));
 
-	if(MediaCapture)
+	if (MediaCapture)
 	{
 		MediaCapture->OnStateChangedNative.RemoveAll(this);
 	}
@@ -123,7 +125,7 @@ void UVCamPixelStreamingSession::SetupCapture()
 
 void UVCamPixelStreamingSession::OnCaptureStateChanged()
 {
-	if(!MediaCapture || !MediaOutput)
+	if (!MediaCapture || !MediaOutput)
 	{
 		return;
 	}
@@ -135,7 +137,7 @@ void UVCamPixelStreamingSession::OnCaptureStateChanged()
 			MediaOutput->StartStreaming();
 			break;
 		case EMediaCaptureState::Stopped:
-			if(MediaCapture->WasViewportResized())
+			if (MediaCapture->WasViewportResized())
 			{
 				UE_LOG(LogPixelStreamingVCam, Log, TEXT("Pixel Streaming VCam capture was stopped due to resize, going to restart capture."));
 				// If it was stopped and viewport resized we assume resize caused the stop, so try a restart of capture here.
@@ -209,30 +211,24 @@ void UVCamPixelStreamingSession::SetupCustomInputHandling()
 		UE_LOG(LogPixelStreamingVCam, Log, TEXT("InputChannel callback - Routing input to active viewport"));
 	}
 
-	if(MediaOutput)
+	if (MediaOutput)
 	{
-		IPixelStreamingModule& PixelStreamingModule = IPixelStreamingModule::Get();
-		typedef Protocol::EPixelStreamingMessageTypes EType;		
-		Protocol::EPixelStreamingMessageDirection MessageDirection = Protocol::EPixelStreamingMessageDirection::ToStreamer;
-		TMap<FString, Protocol::FPixelStreamingInputMessage> Protocol = PixelStreamingModule.GetProtocol().ToStreamerProtocol;
+		IPixelStreamingInputModule& PixelStreamingInputModule = IPixelStreamingInputModule::Get();
+		typedef EPixelStreamingMessageTypes EType;
+		EPixelStreamingMessageDirection MessageDirection = EPixelStreamingMessageDirection::ToStreamer;
+		TMap<FString, FPixelStreamingInputMessage> Protocol = FPixelStreamingInputProtocol::ToStreamerProtocol;
 		/*
-		* ====================
-		* ARKit Transform
-		* ====================
-		*/
-		Protocol::FPixelStreamingInputMessage ARKitMessage = Protocol::FPixelStreamingInputMessage(100, 72, {
-			// 4x4 Transform
-			EType::Float, EType::Float, EType::Float, EType::Float,
-			EType::Float, EType::Float, EType::Float, EType::Float,
-			EType::Float, EType::Float, EType::Float, EType::Float,
-			EType::Float, EType::Float, EType::Float, EType::Float,
-			// Timestamp
-			EType::Double
-		});
-		
-		const TFunction<void(FMemoryReader)>& ARKitHandler = [this](FMemoryReader Ar) { 
-
-			if(!EnableARKitTracking)
+		 * ====================
+		 * ARKit Transform
+		 * ====================
+		 */
+		FPixelStreamingInputMessage ARKitMessage = FPixelStreamingInputMessage(100, { // 4x4 Transform
+																						EType::Float, EType::Float, EType::Float, EType::Float, EType::Float, EType::Float, EType::Float, EType::Float, EType::Float, EType::Float, EType::Float, EType::Float, EType::Float, EType::Float, EType::Float, EType::Float,
+																						// Timestamp
+																						EType::Double });
+
+		const TFunction<void(FMemoryReader)>& ARKitHandler = [this](FMemoryReader Ar) {
+			if (!EnableARKitTracking)
 			{
 				return;
 			}
@@ -254,12 +250,12 @@ void UVCamPixelStreamingSession::SetupCustomInputHandling()
 			double Timestamp;
 			Ar << Timestamp;
 
-			if(TSharedPtr<FPixelStreamingLiveLinkSource> LiveLinkSource = UVCamPixelStreamingSubsystem::Get()->TryGetLiveLinkSource(this))
+			if (TSharedPtr<FPixelStreamingLiveLinkSource> LiveLinkSource = UVCamPixelStreamingSubsystem::Get()->TryGetLiveLinkSource(this))
 			{
 				LiveLinkSource->PushTransformForSubject(GetFName(), FTransform(ARKitMatrix), Timestamp);
 			}
 		};
-		PixelStreamingModule.RegisterMessage(MessageDirection, "ARKitTransform", ARKitMessage, ARKitHandler);
+		PixelStreamingInputModule.RegisterMessage(MessageDirection, "ARKitTransform", ARKitMessage, ARKitHandler);
 	}
 	else
 	{
@@ -269,7 +265,7 @@ void UVCamPixelStreamingSession::SetupCustomInputHandling()
 
 void UVCamPixelStreamingSession::StartCapture()
 {
-	if(!MediaCapture)
+	if (!MediaCapture)
 	{
 		return;
 	}
diff --git a/Engine/Plugins/Media/HardwareEncoders/Source/EncoderAMF/Private/Amf_EncoderModule.cpp b/Engine/Plugins/Media/HardwareEncoders/Source/EncoderAMF/Private/Amf_EncoderModule.cpp
--- a/Engine/Plugins/Media/HardwareEncoders/Source/EncoderAMF/Private/Amf_EncoderModule.cpp
+++ b/Engine/Plugins/Media/HardwareEncoders/Source/EncoderAMF/Private/Amf_EncoderModule.cpp
@@ -51,7 +51,13 @@ public:
 					IVulkanDynamicRHI::AddEnabledDeviceExtensionsAndLayers(ExtentionsToAdd, TArray<const ANSICHAR*>());
 				}
 
-				FCoreDelegates::OnPostEngineInit.AddLambda([]() {FVideoEncoderAmf_H264::Register(FVideoEncoderFactory::Get());});
+				FCoreDelegates::OnPostEngineInit.AddLambda([]()
+				{
+					if (IsRHIDeviceAMD())
+					{
+						FVideoEncoderAmf_H264::Register(FVideoEncoderFactory::Get());
+					}
+				});
 				
 				AMFStarted = true;
 			}
diff --git a/Engine/Plugins/Media/HardwareEncoders/Source/EncoderNVENC/Private/NVENC_EncoderModule.cpp b/Engine/Plugins/Media/HardwareEncoders/Source/EncoderNVENC/Private/NVENC_EncoderModule.cpp
--- a/Engine/Plugins/Media/HardwareEncoders/Source/EncoderNVENC/Private/NVENC_EncoderModule.cpp
+++ b/Engine/Plugins/Media/HardwareEncoders/Source/EncoderNVENC/Private/NVENC_EncoderModule.cpp
@@ -39,7 +39,13 @@ public:
 					IVulkanDynamicRHI::AddEnabledDeviceExtensionsAndLayers(ExtentionsToAdd, TArray<const ANSICHAR*>());
 				}
 
-				FModuleManager::LoadModuleChecked<FCUDAModule>("CUDA").OnPostCUDAInit.AddLambda([]() {FVideoEncoderNVENC_H264::Register(FVideoEncoderFactory::Get());});
+				FModuleManager::LoadModuleChecked<FCUDAModule>("CUDA").OnPostCUDAInit.AddLambda([]()
+				{
+					if (IsRHIDeviceNVIDIA())
+					{
+						FVideoEncoderNVENC_H264::Register(FVideoEncoderFactory::Get());
+					}
+				});
 			}
 		}
 	}
diff --git a/Engine/Plugins/Media/PixelCapture/Source/PixelCapture/PixelCapture.build.cs b/Engine/Plugins/Media/PixelCapture/Source/PixelCapture/PixelCapture.build.cs
--- a/Engine/Plugins/Media/PixelCapture/Source/PixelCapture/PixelCapture.build.cs
+++ b/Engine/Plugins/Media/PixelCapture/Source/PixelCapture/PixelCapture.build.cs
@@ -14,11 +14,15 @@ namespace UnrealBuildTool.Rules
 			var EngineDir = Path.GetFullPath(Target.RelativeEnginePath);
 
 			// This is so for game projects using our public headers don't have to include extra modules they might not know about.
-			PublicDependencyModuleNames.AddRange(new string[] {});
+			PublicDependencyModuleNames.AddRange(new string[] {
+				"RHI",
+			});
 
+			
 			// NOTE: General rule is not to access the private folder of another module
+			string WebRtcIncludeDir = Path.Combine(EngineDir, "Source/ThirdParty/WebRTC/4664/Include");
 			PrivateIncludePaths.AddRange(new string[] {
-				Path.Combine(EngineDir, "Source/ThirdParty/WebRTC/4147/Include/third_party/libyuv/include"),
+				Path.Combine(WebRtcIncludeDir, "third_party/libyuv/include"),
 				Path.Combine(EngineDir, "Source/Runtime/Renderer/Private"),
 			});
 
@@ -28,12 +32,8 @@ namespace UnrealBuildTool.Rules
 				"PixelCaptureShaders",
 				"RenderCore",
 				"Renderer",
-				"RHI",
 				"WebRTC",
 			});
-
-			// required for casting UE4 BackBuffer to Vulkan Texture2D for NvEnc
-			PrivateDependencyModuleNames.AddRange(new string[] {});
 		}
 	}
 }
diff --git a/Engine/Plugins/Media/PixelCapture/Source/PixelCapture/Private/PixelCaptureCapturer.cpp b/Engine/Plugins/Media/PixelCapture/Source/PixelCapture/Private/PixelCaptureCapturer.cpp
--- a/Engine/Plugins/Media/PixelCapture/Source/PixelCapture/Private/PixelCaptureCapturer.cpp
+++ b/Engine/Plugins/Media/PixelCapture/Source/PixelCapture/Private/PixelCaptureCapturer.cpp
@@ -37,7 +37,7 @@ void FPixelCaptureCapturer::Capture(const IPixelCaptureInputFrame& InputFrame)
 	checkf(InputWidth == ExpectedInputWidth && InputHeight == ExpectedInputHeight, TEXT("Capturer input resolution changes are not supported"));
 
 	CurrentOutputBuffer = Buffer->LockProduceBuffer();
-	if (CurrentOutputBuffer == nullptr)
+	if (!CurrentOutputBuffer)
 	{
 		UE_LOG(LogPixelCapture, Error, TEXT("Failed to obtain a produce buffer."));
 		return;
@@ -46,6 +46,7 @@ void FPixelCaptureCapturer::Capture(const IPixelCaptureInputFrame& InputFrame)
 	InitMetadata(InputFrame.Metadata.Copy());
 	StartTime = FPlatformTime::Cycles64();
 	CPUStartTime = 0;
+	GPUEnqueueTime = 0;
 	GPUStartTime = 0;
 
 	BeginProcess(InputFrame, CurrentOutputBuffer.Get());
@@ -74,6 +75,8 @@ void FPixelCaptureCapturer::MarkCPUWorkEnd()
 {
 	CurrentOutputBuffer->Metadata.CaptureProcessCPUTime += FPlatformTime::Cycles64() - CPUStartTime;
 	CPUStartTime = 0;
+
+	GPUEnqueueTime = FPlatformTime::Cycles64();
 }
 
 void FPixelCaptureCapturer::MarkGPUWorkStart()
@@ -83,6 +86,9 @@ void FPixelCaptureCapturer::MarkGPUWorkStart()
 		MarkGPUWorkEnd();
 	}
 	GPUStartTime = FPlatformTime::Cycles64();
+	
+	CurrentOutputBuffer->Metadata.CaptureProcessGPUDelay += GPUStartTime - GPUEnqueueTime;
+	GPUEnqueueTime = 0;
 }
 
 void FPixelCaptureCapturer::MarkGPUWorkEnd()
diff --git a/Engine/Plugins/Media/PixelCapture/Source/PixelCapture/Private/PixelCaptureCapturerI420.cpp b/Engine/Plugins/Media/PixelCapture/Source/PixelCapture/Private/PixelCaptureCapturerI420.cpp
--- a/Engine/Plugins/Media/PixelCapture/Source/PixelCapture/Private/PixelCaptureCapturerI420.cpp
+++ b/Engine/Plugins/Media/PixelCapture/Source/PixelCapture/Private/PixelCaptureCapturerI420.cpp
@@ -12,7 +12,7 @@ void FPixelCaptureCapturerI420::Initialize(int32 InputWidth, int32 InputHeight)
 
 IPixelCaptureOutputFrame* FPixelCaptureCapturerI420::CreateOutputBuffer(int32 InputWidth, int32 InputHeight)
 {
-	return new FPixelCaptureOutputFrameI420(MakeShared<FPixelCaptureI420Buffer>(InputWidth, InputHeight));
+	return new FPixelCaptureOutputFrameI420(MakeShared<FPixelCaptureBufferI420>(InputWidth, InputHeight));
 }
 
 void FPixelCaptureCapturerI420::BeginProcess(const IPixelCaptureInputFrame& InputFrame, IPixelCaptureOutputFrame* OutputBuffer)
diff --git a/Engine/Plugins/Media/PixelCapture/Source/PixelCapture/Private/PixelCaptureCapturerI420ToRHI.cpp b/Engine/Plugins/Media/PixelCapture/Source/PixelCapture/Private/PixelCaptureCapturerI420ToRHI.cpp
--- a/Engine/Plugins/Media/PixelCapture/Source/PixelCapture/Private/PixelCaptureCapturerI420ToRHI.cpp
+++ b/Engine/Plugins/Media/PixelCapture/Source/PixelCapture/Private/PixelCaptureCapturerI420ToRHI.cpp
@@ -2,10 +2,12 @@
 
 #include "PixelCaptureCapturerI420ToRHI.h"
 #include "PixelCaptureInputFrameI420.h"
+#include "PixelCaptureInputFrameNV12.h"
 #include "PixelCaptureOutputFrameRHI.h"
 #include "PixelCaptureBufferFormat.h"
 
 #include "Async/Async.h"
+#include "RenderingThread.h"
 
 #include "libyuv/convert.h"
 #include "libyuv/video_common.h"
@@ -26,7 +28,7 @@ FPixelCaptureCapturerI420ToRHI::~FPixelCaptureCapturerI420ToRHI()
 
 void FPixelCaptureCapturerI420ToRHI::Initialize(int32 InputWidth, int32 InputHeight)
 {
-	Dimensions = { InputWidth, InputHeight };
+	Dimensions = {InputWidth, InputHeight};
 
 	const uint32_t Size = webrtc::CalcBufferSize(webrtc::VideoType::kARGB, Dimensions.X, Dimensions.Y);
 	ARGBBuffer = new uint8_t[Size];
@@ -34,7 +36,7 @@ void FPixelCaptureCapturerI420ToRHI::Initialize(int32 InputWidth, int32 InputHei
 	FPixelCaptureCapturer::Initialize(InputWidth, InputHeight);
 }
 
-IPixelCaptureOutputFrame* FPixelCaptureCapturerI420ToRHI::CreateOutputBuffer(int32 InputWidth, int32 InputHeight)
+IPixelCaptureOutputFrame *FPixelCaptureCapturerI420ToRHI::CreateOutputBuffer(int32 InputWidth, int32 InputHeight)
 {
 	FRHITextureCreateDesc TextureDesc =
 		FRHITextureCreateDesc::Create2D(TEXT("FPixelCaptureCapturerI420ToRHI Texture"), InputWidth, InputHeight, EPixelFormat::PF_B8G8R8A8)
@@ -55,39 +57,45 @@ IPixelCaptureOutputFrame* FPixelCaptureCapturerI420ToRHI::CreateOutputBuffer(int
 	return new FPixelCaptureOutputFrameRHI(GDynamicRHI->RHICreateTexture(TextureDesc));
 }
 
-void FPixelCaptureCapturerI420ToRHI::BeginProcess(const IPixelCaptureInputFrame& InputFrame, IPixelCaptureOutputFrame* OutputBuffer)
+void FPixelCaptureCapturerI420ToRHI::BeginProcess(const IPixelCaptureInputFrame &InputFrame, IPixelCaptureOutputFrame *OutputBuffer)
 {
-	checkf(InputFrame.GetType() == StaticCast<int32>(PixelCaptureBufferFormat::FORMAT_I420), TEXT("Incorrect source frame coming into frame capture process, expected FORMAT_I420"));
-
-	MarkCPUWorkStart();
-
-	const FPixelCaptureInputFrameI420& SourceFrame = StaticCast<const FPixelCaptureInputFrameI420&>(InputFrame);
-	TSharedPtr<FPixelCaptureI420Buffer> I420Buffer = SourceFrame.GetBuffer();
-	FPixelCaptureOutputFrameRHI* OutputH264Buffer = StaticCast<FPixelCaptureOutputFrameRHI*>(OutputBuffer);
-
-	// convert the incoming i420 to argb
-	libyuv::ConvertFromI420(
-		I420Buffer->GetDataY(), I420Buffer->GetStrideY(),
-		I420Buffer->GetDataU(), I420Buffer->GetStrideUV(),
-		I420Buffer->GetDataV(), I420Buffer->GetStrideUV(),
-		ARGBBuffer, 0,
-		I420Buffer->GetWidth(), I420Buffer->GetHeight(),
-		libyuv::FOURCC_ARGB);
-
-	TWeakPtr<FPixelCaptureCapturerI420ToRHI> WeakSelf = StaticCastSharedRef<FPixelCaptureCapturerI420ToRHI>(AsShared());
-	AsyncTask(ENamedThreads::ActualRenderingThread, [WeakSelf, OutputH264Buffer]() {
-		if (TSharedPtr<FPixelCaptureCapturerI420ToRHI> Self = WeakSelf.Pin())
-		{
-			Self->MarkGPUWorkStart();
-
-			// update our staging texture
-			FUpdateTextureRegion2D Region(0, 0, 0, 0, Self->Dimensions.X, Self->Dimensions.Y);
-			RHIUpdateTexture2D(OutputH264Buffer->GetFrameTexture(), 0, Region, Self->Dimensions.X * 4, Self->ARGBBuffer);
-
-			Self->MarkGPUWorkEnd();
-			Self->EndProcess();
-		}
-	});
-
-	MarkCPUWorkEnd();
+	if (InputFrame.GetType() == StaticCast<int32>(PixelCaptureBufferFormat::FORMAT_I420))
+	{
+		MarkCPUWorkStart();
+
+		const FPixelCaptureInputFrameI420 &SourceFrame = StaticCast<const FPixelCaptureInputFrameI420 &>(InputFrame);
+		TSharedPtr<FPixelCaptureBufferI420> I420Buffer = SourceFrame.GetBuffer();
+		FPixelCaptureOutputFrameRHI *OutputTexture = StaticCast<FPixelCaptureOutputFrameRHI *>(OutputBuffer);
+
+		// convert the incoming i420 to argb
+		libyuv::ConvertFromI420(
+			I420Buffer->GetDataY(), I420Buffer->GetStrideY(),
+			I420Buffer->GetDataU(), I420Buffer->GetStrideUV(),
+			I420Buffer->GetDataV(), I420Buffer->GetStrideUV(),
+			ARGBBuffer, 0,
+			I420Buffer->GetWidth(), I420Buffer->GetHeight(),
+			libyuv::FOURCC_ARGB);
+
+		TWeakPtr<FPixelCaptureCapturerI420ToRHI> WeakSelf = StaticCastSharedRef<FPixelCaptureCapturerI420ToRHI>(AsShared());
+		ENQUEUE_RENDER_COMMAND(I420toRHI)
+		([WeakSelf, OutputTexture](FRHICommandListImmediate &RHICmdList) -> void
+		 {
+			if (TSharedPtr<FPixelCaptureCapturerI420ToRHI> Self = WeakSelf.Pin())
+			{
+				Self->MarkGPUWorkStart();
+
+				// update our staging texture
+				FUpdateTextureRegion2D Region(0, 0, 0, 0, Self->Dimensions.X, Self->Dimensions.Y);
+				RHICmdList.UpdateTexture2D(OutputTexture->GetFrameTexture(), 0, Region, Self->Dimensions.X * 4, Self->ARGBBuffer);
+
+				Self->MarkGPUWorkEnd();
+				Self->EndProcess();
+			} });
+
+		MarkCPUWorkEnd();
+	}
+	else
+	{
+		unimplemented();
+	}
 }
diff --git a/Engine/Plugins/Media/PixelCapture/Source/PixelCapture/Private/PixelCaptureCapturerLayered.cpp b/Engine/Plugins/Media/PixelCapture/Source/PixelCapture/Private/PixelCaptureCapturerLayered.cpp
--- a/Engine/Plugins/Media/PixelCapture/Source/PixelCapture/Private/PixelCaptureCapturerLayered.cpp
+++ b/Engine/Plugins/Media/PixelCapture/Source/PixelCapture/Private/PixelCaptureCapturerLayered.cpp
@@ -52,26 +52,32 @@ void FPixelCaptureCapturerLayered::Capture(const IPixelCaptureInputFrame& Source
 		return;
 	}
 
-	FScopeLock LayersLock(&LayersGuard);
-
-	// initial setup. Should only be called on the first frame.
-	// We do this lazily because CreateCaptureProcess is a pure virtual function
-	// so we cannot call it in our constructor. Another option would be to
-	// require the user to call a SetupLayers method or something but that
-	// could be prone to errors.
-	if (LayerCapturers.IsEmpty())
+	// work on a temp list so we dont over lock
+	TArray<TSharedPtr<FPixelCaptureCapturer>> TempLayerCapturers;
 	{
-		for (auto& Scale : LayerScales)
+		FScopeLock LayersLock(&LayersGuard);
+
+		// initial setup. Should only be called on the first frame.
+		// We do this lazily because CreateCaptureProcess is a pure virtual function
+		// so we cannot call it in our constructor. Another option would be to
+		// require the user to call a SetupLayers method or something but that
+		// could be prone to errors.
+		if (LayerCapturers.IsEmpty())
 		{
-			AddLayer(Scale);
+			for (auto& Scale : LayerScales)
+			{
+				AddLayer(Scale);
+			}
 		}
+
+		TempLayerCapturers = LayerCapturers;
 	}
 
 	// set this before calling process since the process might immediately complete
-	PendingLayers = LayerCapturers.Num();
+	PendingLayers = TempLayerCapturers.Num();
 
 	// capture the frame for encoder use
-	for (auto& LayerCapturer : LayerCapturers)
+	for (auto& LayerCapturer : TempLayerCapturers)
 	{
 		LayerCapturer->Capture(SourceFrame);
 	}
diff --git a/Engine/Plugins/Media/PixelCapture/Source/PixelCapture/Private/PixelCaptureCapturerMultiFormat.cpp b/Engine/Plugins/Media/PixelCapture/Source/PixelCapture/Private/PixelCaptureCapturerMultiFormat.cpp
--- a/Engine/Plugins/Media/PixelCapture/Source/PixelCapture/Private/PixelCaptureCapturerMultiFormat.cpp
+++ b/Engine/Plugins/Media/PixelCapture/Source/PixelCapture/Private/PixelCaptureCapturerMultiFormat.cpp
@@ -63,14 +63,14 @@ void FPixelCaptureCapturerMultiFormat::Capture(const IPixelCaptureInputFrame& So
 		// (and their callback events) is ok.
 		// Note: When considering whether the lifetime of this lock can be shorter consider the deadlock issue we had in UE-173694
 
-		TArray<TSharedPtr<FPixelCaptureCapturerLayered>> TempCopy;
+		TArray<TSharedPtr<FPixelCaptureCapturerLayered>> AllFormatCapturers;
 		for (auto& FormatCapturer : FormatCapturers)
 		{
-			TempCopy.Add(FormatCapturer.Value);
+			AllFormatCapturers.Add(FormatCapturer.Value);
 		}
 		// start all the format captures
-		PendingFormats = FormatCapturers.Num();
-		for (auto& Capturer : TempCopy)
+		PendingFormats = AllFormatCapturers.Num();
+		for (auto& Capturer : AllFormatCapturers)
 		{
 			Capturer->Capture(SourceFrame);
 		}
diff --git a/Engine/Plugins/Media/PixelCapture/Source/PixelCapture/Private/PixelCaptureCapturerNV12ToRHI.cpp b/Engine/Plugins/Media/PixelCapture/Source/PixelCapture/Private/PixelCaptureCapturerNV12ToRHI.cpp
new file mode 100644
--- /dev/null
+++ b/Engine/Plugins/Media/PixelCapture/Source/PixelCapture/Private/PixelCaptureCapturerNV12ToRHI.cpp
@@ -0,0 +1,94 @@
+// Copyright Epic Games, Inc. All Rights Reserved.
+
+#include "PixelCaptureCapturerNV12ToRHI.h"
+#include "PixelCaptureInputFrameI420.h"
+#include "PixelCaptureInputFrameNV12.h"
+#include "PixelCaptureOutputFrameRHI.h"
+#include "PixelCaptureBufferFormat.h"
+
+#include "Async/Async.h"
+#include "DynamicRHI.h"
+#include "RenderingThread.h"
+
+#include "libyuv/convert.h"
+#include "libyuv/video_common.h"
+#include "common_video/libyuv/include/webrtc_libyuv.h"
+
+TSharedPtr<FPixelCaptureCapturerNV12ToRHI> FPixelCaptureCapturerNV12ToRHI::Create()
+{
+	return TSharedPtr<FPixelCaptureCapturerNV12ToRHI>(new FPixelCaptureCapturerNV12ToRHI());
+}
+
+FPixelCaptureCapturerNV12ToRHI::~FPixelCaptureCapturerNV12ToRHI()
+{
+	if (R8Buffer)
+	{
+		delete[] R8Buffer;
+	}
+}
+
+void FPixelCaptureCapturerNV12ToRHI::Initialize(int32 InputWidth, int32 InputHeight)
+{
+	Dimensions = { InputWidth, InputHeight };
+
+	const uint32_t Size = (3 * Dimensions.X * Dimensions.Y) >> 1;
+	R8Buffer = new uint8_t[Size];
+
+	FPixelCaptureCapturer::Initialize(InputWidth, InputHeight);
+}
+
+IPixelCaptureOutputFrame* FPixelCaptureCapturerNV12ToRHI::CreateOutputBuffer(int32 InputWidth, int32 InputHeight)
+{
+	FRHITextureCreateDesc TextureDesc =
+		FRHITextureCreateDesc::Create2D(TEXT("FPixelCaptureCapturerNV12ToRHI Texture"), InputWidth, 3 * InputHeight / 2, EPixelFormat::PF_R8)
+			.SetClearValue(FClearValueBinding::None)
+			.SetFlags(ETextureCreateFlags::CPUWritable | ETextureCreateFlags::ShaderResource)
+			.SetInitialState(ERHIAccess::Present)
+			.DetermineInititialState();
+
+	if (RHIGetInterfaceType() == ERHIInterfaceType::Vulkan)
+	{
+		TextureDesc.AddFlags(ETextureCreateFlags::External);
+	}
+	else
+	{
+		TextureDesc.AddFlags(ETextureCreateFlags::Shared);
+	}
+
+	return new FPixelCaptureOutputFrameRHI(GDynamicRHI->RHICreateTexture(TextureDesc), { InputWidth, InputHeight, EPixelFormat::PF_NV12 });
+}
+
+void FPixelCaptureCapturerNV12ToRHI::BeginProcess(const IPixelCaptureInputFrame& InputFrame, IPixelCaptureOutputFrame* OutputBuffer)
+{
+	if (InputFrame.GetType() == StaticCast<int32>(PixelCaptureBufferFormat::FORMAT_NV12))
+	{
+		MarkCPUWorkStart();
+
+		const FPixelCaptureInputFrameNV12& SourceFrame = StaticCast<const FPixelCaptureInputFrameNV12&>(InputFrame);
+		TSharedPtr<FPixelCaptureBufferNV12> NV12Buffer = SourceFrame.GetBuffer();
+		FPixelCaptureOutputFrameRHI* OutputTexture = StaticCast<FPixelCaptureOutputFrameRHI*>(OutputBuffer);
+
+		checkf(OutputTexture->GetFrameTexture()->GetDesc().Format == EPixelFormat::PF_R8, TEXT("Texture Format is incorrect for NV12 storage"))
+
+		TWeakPtr<FPixelCaptureCapturerNV12ToRHI> WeakSelf = StaticCastSharedRef<FPixelCaptureCapturerNV12ToRHI>(AsShared());
+		ENQUEUE_RENDER_COMMAND(PixelCaptureNV12ToRHI)([WeakSelf, OutputTexture, NV12Buffer](FRHICommandListImmediate& RHICmdList) {
+			if (TSharedPtr<FPixelCaptureCapturerNV12ToRHI> Self = WeakSelf.Pin())
+			{
+				Self->MarkGPUWorkStart();
+
+				// update our staging texture
+				FUpdateTextureRegion2D Region(0, 0, 0, 0, Self->Dimensions.X, Self->Dimensions.Y * 1.5);
+				GDynamicRHI->RHIUpdateTexture2D(RHICmdList, OutputTexture->GetFrameTexture(), 0, Region, Self->Dimensions.X, NV12Buffer->GetData());
+
+				Self->MarkGPUWorkEnd();
+				Self->EndProcess();
+			}
+		});
+
+		MarkCPUWorkEnd();
+	}
+	else
+	{
+		unimplemented();
+	}
+}
diff --git a/Engine/Plugins/Media/PixelCapture/Source/PixelCapture/Private/PixelCaptureCapturerRHIRDG.cpp b/Engine/Plugins/Media/PixelCapture/Source/PixelCapture/Private/PixelCaptureCapturerRHIRDG.cpp
--- a/Engine/Plugins/Media/PixelCapture/Source/PixelCapture/Private/PixelCaptureCapturerRHIRDG.cpp
+++ b/Engine/Plugins/Media/PixelCapture/Source/PixelCapture/Private/PixelCaptureCapturerRHIRDG.cpp
@@ -47,6 +47,7 @@ void FPixelCaptureCapturerRHIRDG::BeginProcess(const IPixelCaptureInputFrame& In
 	checkf(InputFrame.GetType() == StaticCast<int32>(PixelCaptureBufferFormat::FORMAT_RHI), TEXT("Incorrect source frame coming into frame capture process."));
 
 	MarkCPUWorkStart();
+	MarkCPUWorkEnd();
 
 	FRHICommandListImmediate& RHICmdList = FRHICommandListExecutor::GetImmediateCommandList();
 
@@ -55,16 +56,15 @@ void FPixelCaptureCapturerRHIRDG::BeginProcess(const IPixelCaptureInputFrame& In
 	FPixelCaptureOutputFrameRHI* OutputH264Buffer = StaticCast<FPixelCaptureOutputFrameRHI*>(OutputBuffer);
 	CopyTextureRDG(RHICmdList, RHISourceFrame.FrameTexture, OutputH264Buffer->GetFrameTexture());
 
+	MarkGPUWorkEnd();
 	OnRHIStageComplete();
 }
 
 void FPixelCaptureCapturerRHIRDG::CheckComplete()
 {
-
 }
 
 void FPixelCaptureCapturerRHIRDG::OnRHIStageComplete()
 {
-	MarkGPUWorkEnd();
 	EndProcess();
 }
diff --git a/Engine/Plugins/Media/PixelCapture/Source/PixelCapture/Private/PixelCaptureCapturerRHIToI420CPU.cpp b/Engine/Plugins/Media/PixelCapture/Source/PixelCapture/Private/PixelCaptureCapturerRHIToI420CPU.cpp
--- a/Engine/Plugins/Media/PixelCapture/Source/PixelCapture/Private/PixelCaptureCapturerRHIToI420CPU.cpp
+++ b/Engine/Plugins/Media/PixelCapture/Source/PixelCapture/Private/PixelCaptureCapturerRHIToI420CPU.cpp
@@ -66,7 +66,7 @@ IPixelCaptureOutputFrame* FPixelCaptureCapturerRHIToI420CPU::CreateOutputBuffer(
 {
 	const int32 Width = InputWidth * Scale;
 	const int32 Height = InputHeight * Scale;
-	return new FPixelCaptureOutputFrameI420(MakeShared<FPixelCaptureI420Buffer>(Width, Height));
+	return new FPixelCaptureOutputFrameI420(MakeShared<FPixelCaptureBufferI420>(Width, Height));
 }
 
 void FPixelCaptureCapturerRHIToI420CPU::BeginProcess(const IPixelCaptureInputFrame& InputFrame, IPixelCaptureOutputFrame* OutputBuffer)
@@ -108,7 +108,7 @@ void FPixelCaptureCapturerRHIToI420CPU::OnRHIStageComplete(IPixelCaptureOutputFr
 	MarkCPUWorkStart();
 
 	FPixelCaptureOutputFrameI420* OutputI420Buffer = StaticCast<FPixelCaptureOutputFrameI420*>(OutputBuffer);
-	TSharedPtr<FPixelCaptureI420Buffer> I420Buffer = OutputI420Buffer->GetI420Buffer();
+	TSharedPtr<FPixelCaptureBufferI420> I420Buffer = OutputI420Buffer->GetI420Buffer();
 	libyuv::ARGBToI420(
 		static_cast<uint8*>(ResultsBuffer),
 		MappedStride * 4,
diff --git a/Engine/Plugins/Media/PixelCapture/Source/PixelCapture/Private/PixelCaptureCapturerRHIToI420Compute.cpp b/Engine/Plugins/Media/PixelCapture/Source/PixelCapture/Private/PixelCaptureCapturerRHIToI420Compute.cpp
--- a/Engine/Plugins/Media/PixelCapture/Source/PixelCapture/Private/PixelCaptureCapturerRHIToI420Compute.cpp
+++ b/Engine/Plugins/Media/PixelCapture/Source/PixelCapture/Private/PixelCaptureCapturerRHIToI420Compute.cpp
@@ -99,7 +99,7 @@ IPixelCaptureOutputFrame* FPixelCaptureCapturerRHIToI420Compute::CreateOutputBuf
 {
 	const int32 Width = InputWidth * Scale;
 	const int32 Height = InputHeight * Scale;
-	return new FPixelCaptureOutputFrameI420(MakeShared<FPixelCaptureI420Buffer>(Width, Height));
+	return new FPixelCaptureOutputFrameI420(MakeShared<FPixelCaptureBufferI420>(Width, Height));
 }
 
 void FPixelCaptureCapturerRHIToI420Compute::BeginProcess(const IPixelCaptureInputFrame& InputFrame, IPixelCaptureOutputFrame* OutputBuffer)
@@ -143,7 +143,7 @@ void FPixelCaptureCapturerRHIToI420Compute::OnRHIStageComplete(IPixelCaptureOutp
 	MarkCPUWorkStart();
 
 	FPixelCaptureOutputFrameI420* OutputI420Buffer = StaticCast<FPixelCaptureOutputFrameI420*>(OutputBuffer);
-	TSharedPtr<FPixelCaptureI420Buffer> I420Buffer = OutputI420Buffer->GetI420Buffer();
+	TSharedPtr<FPixelCaptureBufferI420> I420Buffer = OutputI420Buffer->GetI420Buffer();
 
 	MemCpyStride(I420Buffer->GetMutableDataY(), MappedY, I420Buffer->GetStrideY(), YStride, PlaneYDimensions.Y);
 	MemCpyStride(I420Buffer->GetMutableDataU(), MappedU, I420Buffer->GetStrideUV(), UStride, PlaneUVDimensions.Y);
diff --git a/Engine/Plugins/Media/PixelCapture/Source/PixelCapture/Private/PixelCaptureInputFrameI420.cpp b/Engine/Plugins/Media/PixelCapture/Source/PixelCapture/Private/PixelCaptureInputFrameI420.cpp
--- a/Engine/Plugins/Media/PixelCapture/Source/PixelCapture/Private/PixelCaptureInputFrameI420.cpp
+++ b/Engine/Plugins/Media/PixelCapture/Source/PixelCapture/Private/PixelCaptureInputFrameI420.cpp
@@ -3,7 +3,7 @@
 #include "PixelCaptureInputFrameI420.h"
 #include "PixelCaptureBufferFormat.h"
 
-FPixelCaptureInputFrameI420::FPixelCaptureInputFrameI420(TSharedPtr<FPixelCaptureI420Buffer> Buffer)
+FPixelCaptureInputFrameI420::FPixelCaptureInputFrameI420(TSharedPtr<FPixelCaptureBufferI420> Buffer)
 	: I420Buffer(Buffer)
 {
 	Metadata.SourceTime = FPlatformTime::Cycles64();
@@ -11,7 +11,7 @@ FPixelCaptureInputFrameI420::FPixelCaptureInputFrameI420(TSharedPtr<FPixelCaptur
 
 int32 FPixelCaptureInputFrameI420::GetType() const
 {
-	return static_cast<int32>(PixelCaptureBufferFormat::FORMAT_I420);
+	return I420Buffer->GetFormat();
 }
 
 int32 FPixelCaptureInputFrameI420::GetWidth() const
@@ -24,7 +24,7 @@ int32 FPixelCaptureInputFrameI420::GetHeight() const
 	return I420Buffer->GetHeight();
 }
 
-TSharedPtr<FPixelCaptureI420Buffer> FPixelCaptureInputFrameI420::GetBuffer() const
+TSharedPtr<FPixelCaptureBufferI420> FPixelCaptureInputFrameI420::GetBuffer() const
 {
 	return I420Buffer;
 }
diff --git a/Engine/Plugins/Media/PixelCapture/Source/PixelCapture/Public/PixelCaptureCapturer.h b/Engine/Plugins/Media/PixelCapture/Source/PixelCapture/Public/PixelCaptureCapturer.h
--- a/Engine/Plugins/Media/PixelCapture/Source/PixelCapture/Public/PixelCaptureCapturer.h
+++ b/Engine/Plugins/Media/PixelCapture/Source/PixelCapture/Public/PixelCaptureCapturer.h
@@ -132,6 +132,7 @@ private:
 
 	uint64 StartTime;
 	uint64 CPUStartTime;
+	uint64 GPUEnqueueTime;
 	uint64 GPUStartTime;
 
 	TUniquePtr<UE::PixelCapture::FOutputFrameBuffer> Buffer;
diff --git a/Engine/Plugins/Media/PixelCapture/Source/PixelCapture/Public/PixelCaptureCapturerNV12ToRHI.h b/Engine/Plugins/Media/PixelCapture/Source/PixelCapture/Public/PixelCaptureCapturerNV12ToRHI.h
new file mode 100644
--- /dev/null
+++ b/Engine/Plugins/Media/PixelCapture/Source/PixelCapture/Public/PixelCaptureCapturerNV12ToRHI.h
@@ -0,0 +1,29 @@
+// Copyright Epic Games, Inc. All Rights Reserved.
+
+#pragma once
+
+#include "PixelCaptureCapturer.h"
+
+/**
+ * A basic capturer that will capture NV12 frames to native RHI textures.
+ * Input: FPixelCaptureInputFrameNV12
+ * Output: FPixelCaptureOutputFrameRHI
+ */
+class PIXELCAPTURE_API FPixelCaptureCapturerNV12ToRHI : public FPixelCaptureCapturer, public TSharedFromThis<FPixelCaptureCapturerNV12ToRHI>
+{
+public:
+	static TSharedPtr<FPixelCaptureCapturerNV12ToRHI> Create();
+	virtual ~FPixelCaptureCapturerNV12ToRHI();
+
+protected:
+	virtual FString GetCapturerName() const override { return "NV12ToRHI"; }
+	virtual void Initialize(int32 InputWidth, int32 InputHeight) override;
+	virtual IPixelCaptureOutputFrame* CreateOutputBuffer(int32 InputWidth, int32 InputHeight) override;
+	virtual void BeginProcess(const IPixelCaptureInputFrame& InputFrame, IPixelCaptureOutputFrame* OutputBuffer) override;
+
+private:
+	FIntPoint Dimensions;
+	uint8_t* R8Buffer = nullptr;
+
+	FPixelCaptureCapturerNV12ToRHI() = default;
+};
diff --git a/Engine/Plugins/Media/PixelCapture/Source/PixelCapture/Public/PixelCaptureFrameMetadata.h b/Engine/Plugins/Media/PixelCapture/Source/PixelCapture/Public/PixelCaptureFrameMetadata.h
--- a/Engine/Plugins/Media/PixelCapture/Source/PixelCapture/Public/PixelCaptureFrameMetadata.h
+++ b/Engine/Plugins/Media/PixelCapture/Source/PixelCapture/Public/PixelCaptureFrameMetadata.h
@@ -18,14 +18,30 @@ struct FPixelCaptureFrameMetadata
 	// Capture process timings. Duration not timestamp.
 	uint64 CaptureTime = 0;
 	uint64 CaptureProcessCPUTime = 0;
+	uint64 CaptureProcessGPUDelay = 0;
 	uint64 CaptureProcessGPUTime = 0;
 
+	// Display process timings. Duration not timestamp.
+	uint64 DisplayTime = 0;
+
 	// Frame use timings (can happen multiple times. ie. we are consuming frames faster than producing them)
-	uint32 UseCount = 0; // how many times the frame has been fed to the encoder
+	uint32 UseCount = 0; // how many times the frame has been fed to the encoder or decoder
+
+	// Encode Timings
 	uint64 FirstEncodeStartTime = 0;
 	uint64 LastEncodeStartTime = 0;
 	uint64 LastEncodeEndTime = 0;
 
+	// Packet Timings
+	uint64 FirstPacketizationStartTime = 0;
+	uint64 LastPacketizationStartTime = 0;
+	uint64 LastPacketizationEndTime = 0;
+
+	// Decode Timings
+	uint64 FirstDecodeStartTime = 0;
+	uint64 LastDecodeStartTime = 0;
+	uint64 LastDecodeEndTime = 0;
+
 	// wanted this to be explicit with a name
 	FPixelCaptureFrameMetadata Copy() const { return *this; }
 };
diff --git a/Engine/Plugins/Media/PixelCapture/Source/PixelCapture/Public/PixelCaptureI420Buffer.h b/Engine/Plugins/Media/PixelCapture/Source/PixelCapture/Public/PixelCaptureI420Buffer.h
--- a/Engine/Plugins/Media/PixelCapture/Source/PixelCapture/Public/PixelCaptureI420Buffer.h
+++ b/Engine/Plugins/Media/PixelCapture/Source/PixelCapture/Public/PixelCaptureI420Buffer.h
@@ -2,109 +2,7 @@
 
 #pragma once
 
-#include "CoreMinimal.h"
+#include "PixelCaptureBufferI420.h"
 
-/**
- * Buffer class for holding I420 frame data.
- */
-class PIXELCAPTURE_API FPixelCaptureI420Buffer
-{
-public:
-	/**
-	 * Create an empty buffer of the specified size.
-	 */
-	FPixelCaptureI420Buffer(int InWidth, int InHeight);
-
-	FPixelCaptureI420Buffer(const FPixelCaptureI420Buffer& Other) = default;
-	~FPixelCaptureI420Buffer() = default;
-
-	/**
-	 * Get the width of the frame
-	 */
-	int GetWidth() const { return Width; }
-
-	/**
-	 * Get the height of the frame
-	 */
-	int GetHeight() const { return Height; }
-
-	/**
-	 * Get the stride of the Y plane
-	 */
-	int GetStrideY() const { return StrideY; }
-
-	/**
-	 * Get the stride of the U and V planes
-	 */
-	int GetStrideUV() const { return StrideUV; }
-
-	/**
-	 * Gets a const pointer to the beginning of entire buffer for reading.
-	 */
-	const uint8_t* GetData() const;
-
-	/**
-	 * Gets a const pointer to the beginning of the Y plane for reading.
-	 */
-	const uint8_t* GetDataY() const;
-
-	/**
-	 * Gets a const pointer to the beginning of the U plane for reading.
-	 */
-	const uint8_t* GetDataU() const;
-
-	/**
-	 * Gets a const pointer to the beginning of the V plane for reading.
-	 */
-	const uint8_t* GetDataV() const;
-
-	/**
-	 * Gets a pointer to the beginning of the entire buffer for editing.
-	 */
-	uint8_t* GetMutableData();
-
-	/**
-	 * Gets a pointer to the beginning of the Y plane for editing.
-	 */
-	uint8_t* GetMutableDataY();
-
-	/**
-	 * Gets a pointer to the beginning of the U plane for editing.
-	 */
-	uint8_t* GetMutableDataU();
-
-	/**
-	 * Gets a pointer to the beginning of the V plane for editing.
-	 */
-	uint8_t* GetMutableDataV();
-
-	/**
-	 * Gets the size of the Y plane in bytes.
-	 */
-	int GetDataSizeY() const;
-
-	/**
-	 * Gets the size of both the U and V planes in bytes.
-	 */
-	int GetDataSizeUV() const;
-
-	/**
-	 * Gets the total size of this buffer.
-	 */
-	int GetSize() const { return Data.Num(); }
-
-	/**
-	 * Copies the given buffer into this buffer. Supplied so copies
-	 * can be explicit in code where we want to be verbose.
-	 */
-	void Copy(const FPixelCaptureI420Buffer& Other) { *this = Other; }
-
-private:
-	int Width = 0;
-	int Height = 0;
-
-	int StrideY = 0;
-	int StrideUV = 0;
-
-	TArray<uint8_t> Data;
-};
+UE_DEPRECATED(5.2, "Use FPixelCaptureBufferI420 instead.")
+typedef FPixelCaptureBufferI420 FPixelCaptureI420Buffer;
\ No newline at end of file
diff --git a/Engine/Plugins/Media/PixelCapture/Source/PixelCapture/Public/PixelCaptureInputFrameI420.h b/Engine/Plugins/Media/PixelCapture/Source/PixelCapture/Public/PixelCaptureInputFrameI420.h
--- a/Engine/Plugins/Media/PixelCapture/Source/PixelCapture/Public/PixelCaptureInputFrameI420.h
+++ b/Engine/Plugins/Media/PixelCapture/Source/PixelCapture/Public/PixelCaptureInputFrameI420.h
@@ -3,7 +3,7 @@
 #pragma once
 
 #include "IPixelCaptureInputFrame.h"
-#include "PixelCaptureI420Buffer.h"
+#include "PixelCaptureBufferI420.h"
 
 /**
  * A basic input frame for the Capture system that wraps a I420 buffer.
@@ -11,15 +11,15 @@
 class PIXELCAPTURE_API FPixelCaptureInputFrameI420 : public IPixelCaptureInputFrame
 {
 public:
-	FPixelCaptureInputFrameI420(TSharedPtr<FPixelCaptureI420Buffer> Buffer);
+	FPixelCaptureInputFrameI420(TSharedPtr<FPixelCaptureBufferI420> Buffer);
 	virtual ~FPixelCaptureInputFrameI420() = default;
 
 	virtual int32 GetType() const override;
 	virtual int32 GetWidth() const override;
 	virtual int32 GetHeight() const override;
 
-	TSharedPtr<FPixelCaptureI420Buffer> GetBuffer() const;
+	TSharedPtr<FPixelCaptureBufferI420> GetBuffer() const;
 
 private:
-	TSharedPtr<FPixelCaptureI420Buffer> I420Buffer;
+	TSharedPtr<FPixelCaptureBufferI420> I420Buffer;
 };
diff --git a/Engine/Plugins/Media/PixelCapture/Source/PixelCapture/Public/PixelCaptureInputFrameNV12.h b/Engine/Plugins/Media/PixelCapture/Source/PixelCapture/Public/PixelCaptureInputFrameNV12.h
new file mode 100644
--- /dev/null
+++ b/Engine/Plugins/Media/PixelCapture/Source/PixelCapture/Public/PixelCaptureInputFrameNV12.h
@@ -0,0 +1,25 @@
+// Copyright Epic Games, Inc. All Rights Reserved.
+
+#pragma once
+
+#include "IPixelCaptureInputFrame.h"
+#include "PixelCaptureBufferNV12.h"
+
+/**
+ * A basic input frame for the Capture system that wraps a I420 buffer.
+ */
+class PIXELCAPTURE_API FPixelCaptureInputFrameNV12 : public IPixelCaptureInputFrame
+{
+public:
+	FPixelCaptureInputFrameNV12(TSharedPtr<FPixelCaptureBufferNV12> Buffer);
+	virtual ~FPixelCaptureInputFrameNV12() = default;
+
+	virtual int32 GetType() const override;
+	virtual int32 GetWidth() const override;
+	virtual int32 GetHeight() const override;
+
+	TSharedPtr<FPixelCaptureBufferNV12> GetBuffer() const;
+
+private:
+	TSharedPtr<FPixelCaptureBufferNV12> NV12Buffer;
+};
diff --git a/Engine/Plugins/Media/PixelCapture/Source/PixelCapture/Public/PixelCaptureOutputFrameI420.h b/Engine/Plugins/Media/PixelCapture/Source/PixelCapture/Public/PixelCaptureOutputFrameI420.h
--- a/Engine/Plugins/Media/PixelCapture/Source/PixelCapture/Public/PixelCaptureOutputFrameI420.h
+++ b/Engine/Plugins/Media/PixelCapture/Source/PixelCapture/Public/PixelCaptureOutputFrameI420.h
@@ -3,7 +3,7 @@
 #pragma once
 
 #include "IPixelCaptureOutputFrame.h"
-#include "PixelCaptureI420Buffer.h"
+#include "PixelCaptureBufferI420.h"
 
 /**
  * A basic output frame from the Capture system that wraps a I420 buffer.
@@ -11,7 +11,7 @@
 class PIXELCAPTURE_API FPixelCaptureOutputFrameI420 : public IPixelCaptureOutputFrame
 {
 public:
-	FPixelCaptureOutputFrameI420(TSharedPtr<FPixelCaptureI420Buffer> InI420Buffer)
+	FPixelCaptureOutputFrameI420(TSharedPtr<FPixelCaptureBufferI420> InI420Buffer)
 		: I420Buffer(InI420Buffer)
 	{
 	}
@@ -20,8 +20,8 @@ public:
 	virtual int32 GetWidth() const override { return I420Buffer->GetWidth(); }
 	virtual int32 GetHeight() const override { return I420Buffer->GetHeight(); }
 
-	TSharedPtr<FPixelCaptureI420Buffer> GetI420Buffer() const { return I420Buffer; }
+	TSharedPtr<FPixelCaptureBufferI420> GetI420Buffer() const { return I420Buffer; }
 
 private:
-	TSharedPtr<FPixelCaptureI420Buffer> I420Buffer;
+	TSharedPtr<FPixelCaptureBufferI420> I420Buffer;
 };
