tidy platform references

#rnx

diff --git a/Engine/Plugins/Experimental/TextureMediaPlayer/Source/TextureMediaPlayerFactory/Private/TextureMediaPlayerFactoryModule.cpp b/Engine/Plugins/Experimental/TextureMediaPlayer/Source/TextureMediaPlayerFactory/Private/TextureMediaPlayerFactoryModule.cpp
--- a/Engine/Plugins/Experimental/TextureMediaPlayer/Source/TextureMediaPlayerFactory/Private/TextureMediaPlayerFactoryModule.cpp
+++ b/Engine/Plugins/Experimental/TextureMediaPlayer/Source/TextureMediaPlayerFactory/Private/TextureMediaPlayerFactoryModule.cpp
@@ -97,20 +97,22 @@ public:
 		FModuleManager::Get().LoadModule(TEXT("TextureMediaPlayer"));
 
 		// supported platforms
-		SupportedPlatforms.Add(TEXT("Windows"));
-		SupportedPlatforms.Add(TEXT("PS4"));
-		SupportedPlatforms.Add(TEXT("Switch"));
-		SupportedPlatforms.Add(TEXT("XboxOne"));
-		SupportedPlatforms.Add(TEXT("XboxOneGDK"));
-		SupportedPlatforms.Add(TEXT("PS5"));
-		SupportedPlatforms.Add(TEXT("XSX"));
+		auto MediaModule = FModuleManager::GetModulePtr<IMediaModule>("Media");
+		if (MediaModule != nullptr)
+		{		
+			AddSupportedPlatform(FGuid(0xd1d5f296, 0xff834a87, 0xb20faaa9, 0xd6b8e9a6));
+			AddSupportedPlatform(FGuid(0xb80decd6, 0x997a4b3f, 0x92063970, 0xe572c0db));
+			AddSupportedPlatform(FGuid(0x30ebce04, 0x2c8247bd, 0xaf873017, 0x5a27ed45));
+			AddSupportedPlatform(FGuid(0xb596ce6f, 0xd8324a9c, 0x84e9f880, 0x21322535));
+			AddSupportedPlatform(FGuid(0x941259d5, 0x0a2746aa, 0xadc0ba84, 0x4790ad8a));
+			AddSupportedPlatform(FGuid(0xb67dd9c6, 0x77694fd5, 0xb2b0c8bf, 0xe0c1c673));
+			AddSupportedPlatform(FGuid(0xccf05903, 0x822b47e1, 0xb2236a28, 0xdfd78817));
+		}
 
 		// supported schemes
 		SupportedUriSchemes.Add(TEXT("texture"));
 
 		// register player factory
-		auto MediaModule = FModuleManager::LoadModulePtr<IMediaModule>("Media");
-
 		if (MediaModule != nullptr)
 		{
 			MediaModule->RegisterPlayerFactory(*this);
@@ -129,6 +131,18 @@ public:
 	}
 
 private:
+	void AddSupportedPlatform(const FGuid& PlatformGuid)
+	{
+		auto MediaModule = FModuleManager::GetModulePtr<IMediaModule>("Media");
+		check(MediaModule);
+
+		FName PlatformName = MediaModule->GetPlatformName(PlatformGuid);
+		if (!PlatformName.IsNone())
+		{
+			SupportedPlatforms.Add(PlatformName.ToString());
+		}
+	}
+
 	/** List of platforms that the media player support. */
 	TArray<FString> SupportedPlatforms;
 
