update supported platforms

#rnx

diff --git a/Engine/Plugins/Experimental/TextureMediaPlayer/Source/TextureMediaPlayerFactory/Private/TextureMediaPlayerFactoryModule.cpp b/Engine/Plugins/Experimental/TextureMediaPlayer/Source/TextureMediaPlayerFactory/Private/TextureMediaPlayerFactoryModule.cpp
--- a/Engine/Plugins/Experimental/TextureMediaPlayer/Source/TextureMediaPlayerFactory/Private/TextureMediaPlayerFactoryModule.cpp
+++ b/Engine/Plugins/Experimental/TextureMediaPlayer/Source/TextureMediaPlayerFactory/Private/TextureMediaPlayerFactoryModule.cpp
@@ -107,6 +107,7 @@ public:
 			AddSupportedPlatform(FGuid(0x941259d5, 0x0a2746aa, 0xadc0ba84, 0x4790ad8a));
 			AddSupportedPlatform(FGuid(0xb67dd9c6, 0x77694fd5, 0xb2b0c8bf, 0xe0c1c673));
 			AddSupportedPlatform(FGuid(0xccf05903, 0x822b47e1, 0xb2236a28, 0xdfd78817));
+			AddSupportedPlatform(FGuid(0xae496f22, 0x95534328, 0xbd035b4c, 0x919dc51a));
 		}
 
 		// supported schemes
diff --git a/Engine/Plugins/Media/ElectraPlayer/Source/ElectraPlayerFactory/Private/ElectraPlayerFactoryModule.cpp b/Engine/Plugins/Media/ElectraPlayer/Source/ElectraPlayerFactory/Private/ElectraPlayerFactoryModule.cpp
--- a/Engine/Plugins/Media/ElectraPlayer/Source/ElectraPlayerFactory/Private/ElectraPlayerFactoryModule.cpp
+++ b/Engine/Plugins/Media/ElectraPlayer/Source/ElectraPlayerFactory/Private/ElectraPlayerFactoryModule.cpp
@@ -133,6 +133,7 @@ public:
 		AddSupportedPlatform(FGuid(0x115de4fe, 0x241b465b, 0x970a872f, 0x3167492a));
 		AddSupportedPlatform(FGuid(0xc0b45a33, 0x9de340c7, 0xbce24c47, 0x15c3babf));
         AddSupportedPlatform(FGuid(0xa478294f, 0xbd0d4ec0, 0x8830b6d4, 0xd219c1a4));
+		AddSupportedPlatform(FGuid(0xae496f22, 0x95534328, 0xbd035b4c, 0x919dc51a));
 #endif
 
 		// supported schemes
