Media: Removed empty overloaded methods

#rnx

diff --git a/Engine/Plugins/Experimental/TextureMediaPlayer/Source/TextureMediaPlayer/Private/TextureMediaPlayer.cpp b/Engine/Plugins/Experimental/TextureMediaPlayer/Source/TextureMediaPlayer/Private/TextureMediaPlayer.cpp
--- a/Engine/Plugins/Experimental/TextureMediaPlayer/Source/TextureMediaPlayer/Private/TextureMediaPlayer.cpp
+++ b/Engine/Plugins/Experimental/TextureMediaPlayer/Source/TextureMediaPlayer/Private/TextureMediaPlayer.cpp
@@ -163,15 +163,6 @@ void FTextureMediaPlayer::Close()
 	MostRecentFrameCount = 0;
 }
 
-//-----------------------------------------------------------------------------
-/**
- *	TODO: This would be a better clock. This is the actual time of the audio renderer.
- */
-void FTextureMediaPlayer::SetLastAudioRenderedSampleTime(FTimespan SampleTime)
-{
-	// No-op.
-}
-
 //-----------------------------------------------------------------------------
 /**
  *	Tick the player from the game thread
@@ -187,7 +178,7 @@ void FTextureMediaPlayer::TickInput(FTimespan DeltaTime, FTimespan Timecode)
 	}
 
 	// Handle state changes as we receive frames.
-	if (State != EMediaState::Error)	
+	if (State != EMediaState::Error)
 	{
 		// Has a new frame arrived?
 		if (MostRecentFrameCount != CurrentFrameCount)
@@ -435,7 +426,7 @@ void FTextureMediaPlayer::OnFrame(const TArray<uint8>& TextureBuffer, FIntPoint
 		// Ownership of the decoded frame is thereby transfered to the texture sample.
 		TextureSample->Initialize(VideoDecoderOutput);
 		++MostRecentFrameCount;
-		MediaSamples->AddVideo(TextureSample);		
+		MediaSamples->AddVideo(TextureSample);
 	}
 }
 
@@ -448,7 +439,7 @@ void FTextureMediaPlayer::OnFrame(FTextureRHIRef TextureRHIRef, TRefCountPtr<ID3
 		int32 Width = SizeXYZ.X;
 		int32 SampleHeight = SizeXYZ.Y;
 		int32 Height = (SampleHeight * 2) / 3;
-		
+
 		FVideoDecoderOutputPtr NewVideoDecoderOutput = MakeShared<TextureMediaPlayerVideoDecoderOutput, ESPMode::ThreadSafe>();
 		TextureMediaPlayerVideoDecoderOutput* VideoDecoderOutput = static_cast<TextureMediaPlayerVideoDecoderOutput*>(NewVideoDecoderOutput.Get());
 
diff --git a/Engine/Plugins/Experimental/TextureMediaPlayer/Source/TextureMediaPlayer/Private/TextureMediaPlayer.h b/Engine/Plugins/Experimental/TextureMediaPlayer/Source/TextureMediaPlayer/Private/TextureMediaPlayer.h
--- a/Engine/Plugins/Experimental/TextureMediaPlayer/Source/TextureMediaPlayer/Private/TextureMediaPlayer.h
+++ b/Engine/Plugins/Experimental/TextureMediaPlayer/Source/TextureMediaPlayer/Private/TextureMediaPlayer.h
@@ -82,7 +82,6 @@ public:
 	bool Open(const FString& Url, const IMediaOptions* Options, const FMediaPlayerOptions* PlayerOptions) override;
 	bool Open(const TSharedRef<FArchive, ESPMode::ThreadSafe>& Archive, const FString& OriginalUrl, const IMediaOptions* Options) override;
 	void TickInput(FTimespan DeltaTime, FTimespan Timecode) override;
-	void SetLastAudioRenderedSampleTime(FTimespan SampleTime) override;
 	bool FlushOnSeekStarted() const override
 	{
 		return true;
@@ -97,11 +96,11 @@ public:
 	FTimespan GetTime() const override;
 
 	// ITextureMediaPlayer interface
-	virtual void OnFrame(const TArray<uint8>& TextureBuffer, FIntPoint Size) override; // TODO (aidan.possemiers) 
+	virtual void OnFrame(const TArray<uint8>& TextureBuffer, FIntPoint Size) override; // TODO (aidan.possemiers)
 #if PLATFORM_WINDOWS
 	virtual void OnFrame(FTextureRHIRef TextureRHIRef, TRefCountPtr<ID3D12Fence> D3DFence, uint64 FenceValue) override;
 #else
-	virtual void OnFrame(FTextureRHIRef TextureRHIRef, FGPUFenceRHIRef Fence, uint64 FenceValue) override { unimplemented(); }; // TODO (aidan.possemiers) 
+	virtual void OnFrame(FTextureRHIRef TextureRHIRef, FGPUFenceRHIRef Fence, uint64 FenceValue) override { unimplemented(); }; // TODO (aidan.possemiers)
 #endif // PLATFORM_WINDOWS
 
 private:
diff --git a/Engine/Plugins/Media/ElectraPlayer/Source/ElectraPlayerPlugin/Private/ElectraPlayerPlugin.cpp b/Engine/Plugins/Media/ElectraPlayer/Source/ElectraPlayerPlugin/Private/ElectraPlayerPlugin.cpp
--- a/Engine/Plugins/Media/ElectraPlayer/Source/ElectraPlayerPlugin/Private/ElectraPlayerPlugin.cpp
+++ b/Engine/Plugins/Media/ElectraPlayer/Source/ElectraPlayerPlugin/Private/ElectraPlayerPlugin.cpp
@@ -1119,9 +1119,3 @@ bool FElectraPlayerPlugin::SetVideoTrackFrameRate(int32 TrackIndex, int32 Format
 {
 	return false;
 }
-
-
-void FElectraPlayerPlugin::SetLastAudioRenderedSampleTime(FTimespan SampleTime)
-{
-	// No-op.
-}
diff --git a/Engine/Plugins/Media/ElectraPlayer/Source/ElectraPlayerPlugin/Private/ElectraPlayerPlugin.h b/Engine/Plugins/Media/ElectraPlayer/Source/ElectraPlayerPlugin/Private/ElectraPlayerPlugin.h
--- a/Engine/Plugins/Media/ElectraPlayer/Source/ElectraPlayerPlugin/Private/ElectraPlayerPlugin.h
+++ b/Engine/Plugins/Media/ElectraPlayer/Source/ElectraPlayerPlugin/Private/ElectraPlayerPlugin.h
@@ -98,7 +98,6 @@ public:
 	bool Open(const FString& Url, const IMediaOptions* Options, const FMediaPlayerOptions* PlayerOptions) override;
 	bool Open(const TSharedRef<FArchive, ESPMode::ThreadSafe>& Archive, const FString& OriginalUrl, const IMediaOptions* Options) override;
 	void TickInput(FTimespan DeltaTime, FTimespan Timecode) override;
-	void SetLastAudioRenderedSampleTime(FTimespan SampleTime) override;
 
 	FVariant GetMediaInfo(FName InInfoName) const override;
 
